Move DatasmithInterchange plugin from //UE5/Main/Engine/Restricted/NotForLicensees/Plugins/ to //UE5/Main/Engine/Plugins/Experimental/Enterprise/

#rnx

diff --git a/Engine/Plugins/Experimental/Enterprise/DatasmithInterchange/DatasmithInterchange.uplugin b/Engine/Plugins/Experimental/Enterprise/DatasmithInterchange/DatasmithInterchange.uplugin
new file mode 100644
--- /dev/null
+++ b/Engine/Plugins/Experimental/Enterprise/DatasmithInterchange/DatasmithInterchange.uplugin
@@ -0,0 +1,43 @@
+{
+	"FileVersion" : 1,
+	"Version" : 1,
+	"VersionName" : "1.0",
+	"FriendlyName" : "Datasmith Interchange",
+	"Description" : "Interchange Importer for Datasmith.",
+	"Category" : "Importers",
+	"CreatedBy" : "Epic Games, Inc.",
+	"CreatedByURL" : "http://epicgames.com",
+	"DocsURL" : "",
+	"MarketplaceURL" : "",
+	"SupportURL" : "",
+	"EnabledByDefault" : false,
+	"CanContainContent" : false,
+	"IsBetaVersion" : true,
+	"Installed" : false,
+	"Modules": [
+		{
+			"Name": "DatasmithInterchange",
+			"Type": "Runtime",
+			"LoadingPhase": "Default",
+			"PlatformAllowList": [
+				"Win64",
+				"Linux",
+				"Mac"
+			]
+		}
+	],
+	"Plugins": [
+		{
+			"Name": "DatasmithImporter",
+			"Enabled": true
+		},
+		{
+			"Name": "DatasmithContent",
+			"Enabled":  true
+		},
+		{
+			"Name": "Interchange",
+			"Enabled": true
+		}
+	]
+}
\ No newline at end of file
diff --git a/Engine/Plugins/Experimental/Enterprise/DatasmithInterchange/Source/DatasmithInterchange/DatasmithInterchange.Build.cs b/Engine/Plugins/Experimental/Enterprise/DatasmithInterchange/Source/DatasmithInterchange/DatasmithInterchange.Build.cs
new file mode 100644
--- /dev/null
+++ b/Engine/Plugins/Experimental/Enterprise/DatasmithInterchange/Source/DatasmithInterchange/DatasmithInterchange.Build.cs
@@ -0,0 +1,55 @@
+// Copyright Epic Games, Inc. All Rights Reserved.
+
+namespace UnrealBuildTool.Rules
+{	public class DatasmithInterchange : ModuleRules
+	{
+		public DatasmithInterchange(ReadOnlyTargetRules Target)
+			: base(Target)
+		{
+
+			PrivateDependencyModuleNames.AddRange(
+				new string[]
+				{
+					"DatasmithContent",
+					"DatasmithCore",
+					"DatasmithTranslator",
+					"Engine",
+					"ExternalSource",
+					"InterchangeEngine",
+					"InterchangeNodes",
+					"InterchangeCommonParser",
+					"InterchangeFactoryNodes",
+					"InterchangeImport",
+					"InterchangePipelines",
+					"MeshDescription"
+				}
+			);
+
+			if (Target.bBuildEditor)
+			{
+				PrivateDependencyModuleNames.AddRange(
+					new string[]
+					{
+						"DatasmithImporter",
+						"DesktopPlatform",
+						"InterchangeEditorPipelines",
+						"Slate",
+						"SlateCore",
+						"ToolMenus",
+						"UnrealEd",
+					}
+				);
+			}
+
+			PublicDependencyModuleNames.AddRange(
+				new string[]
+				{
+					"Core",
+					"CoreUObject",
+					"InterchangeCore",
+					"InterchangeImport",
+				}
+			);
+		}
+	}
+}
\ No newline at end of file
diff --git a/Engine/Plugins/Experimental/Enterprise/DatasmithInterchange/Source/DatasmithInterchange/Private/DatasmithInterchangeModule.cpp b/Engine/Plugins/Experimental/Enterprise/DatasmithInterchange/Source/DatasmithInterchange/Private/DatasmithInterchangeModule.cpp
new file mode 100644
--- /dev/null
+++ b/Engine/Plugins/Experimental/Enterprise/DatasmithInterchange/Source/DatasmithInterchange/Private/DatasmithInterchangeModule.cpp
@@ -0,0 +1,179 @@
+// Copyright Epic Games, Inc. All Rights Reserved.
+
+#include "DatasmithInterchangeModule.h"
+
+#include "InterchangeDatasmithAreaLightFactory.h"
+#include "InterchangeDatasmithLog.h"
+#include "InterchangeDatasmithPipeline.h"
+#include "InterchangeDatasmithSceneFactory.h"
+#include "InterchangeDatasmithTranslator.h"
+
+#include "InterchangeReferenceMaterials/DatasmithReferenceMaterialManager.h"
+#include "InterchangeReferenceMaterials/DatasmithC4DMaterialSelector.h"
+#include "InterchangeReferenceMaterials/DatasmithCityEngineMaterialSelector.h"
+#include "InterchangeReferenceMaterials/DatasmithRevitMaterialSelector.h"
+#include "InterchangeReferenceMaterials/DatasmithSketchupMaterialSelector.h"
+#include "InterchangeReferenceMaterials/DatasmithStdMaterialSelector.h"
+
+#include "DatasmithTranslatorManager.h"
+#include "InterchangeGenericAssetsPipeline.h"
+#include "InterchangeGenericMeshPipeline.h"
+#include "InterchangeGenericScenesPipeline.h"
+#include "InterchangeManager.h"
+#include "Logging/LogMacros.h"
+
+#if WITH_EDITOR
+#include "DatasmithImporterModule.h"
+#include "DesktopPlatformModule.h"
+#include "Dialogs/DlgPickPath.h"
+#include "EditorDirectories.h"
+#include "IDesktopPlatform.h"
+#include "ObjectTools.h"
+#include "ToolMenus.h"
+#include "ToolMenuSection.h"
+#endif //WITH_EDITOR
+
+#define LOCTEXT_NAMESPACE "DatasmithInterchange"
+
+
+DEFINE_LOG_CATEGORY(LogInterchangeDatasmith);
+
+class FDatasmithInterchangeModule : public IDatasmithInterchangeModule
+{
+public:
+	virtual void StartupModule() override
+	{
+		using namespace UE::DatasmithInterchange;
+
+		// Load the blueprint asset into memory while wew're on the game thread so that GetAreaLightActorBPClass() can safely be called from other threads.
+		UBlueprint* AreaLightBlueprint = Cast<UBlueprint>(FSoftObjectPath(TEXT("/DatasmithContent/Datasmith/DatasmithArealight.DatasmithArealight")).TryLoad());
+		ensure(AreaLightBlueprint != nullptr);
+
+		UInterchangeManager& InterchangeManager = UInterchangeManager::GetInterchangeManager();
+		InterchangeManager.RegisterTranslator(UInterchangeDatasmithTranslator::StaticClass());
+		InterchangeManager.RegisterFactory(UInterchangeDatasmithSceneFactory::StaticClass());
+		InterchangeManager.RegisterFactory(UInterchangeDatasmithAreaLightFactory::StaticClass());
+
+#if WITH_EDITOR
+		IDatasmithImporterModule& DatasmithImporterModule = IDatasmithImporterModule::Get();
+		DatasmithImporterModule.OnGenerateDatasmithImportMenu().AddRaw(this, &FDatasmithInterchangeModule::ExtendDatasmitMenuOptions);
+#endif //WITH_EDITOR
+
+		FDatasmithReferenceMaterialManager::Create();
+
+		//A minimal set of natively supported reference materials
+		FDatasmithReferenceMaterialManager::Get().RegisterSelector(TEXT("C4D"), MakeShared< FDatasmithC4DMaterialSelector >());
+		FDatasmithReferenceMaterialManager::Get().RegisterSelector(TEXT("Revit"), MakeShared< FDatasmithRevitMaterialSelector >());
+		FDatasmithReferenceMaterialManager::Get().RegisterSelector(TEXT("SketchUp"), MakeShared< FDatasmithSketchUpMaterialSelector >());
+		FDatasmithReferenceMaterialManager::Get().RegisterSelector(TEXT("CityEngine"), MakeShared< FDatasmithCityEngineMaterialSelector >());
+		FDatasmithReferenceMaterialManager::Get().RegisterSelector(TEXT("StdMaterial"), MakeShared< FDatasmithStdMaterialSelector >());
+	}
+
+	virtual void ShutdownModule() override
+	{
+#if WITH_EDITOR
+		if (IDatasmithImporterModule::IsAvailable())
+		{
+			IDatasmithImporterModule& DatasmithImporterModule = IDatasmithImporterModule::Get();
+			DatasmithImporterModule.OnGenerateDatasmithImportMenu().RemoveAll(this);
+		}
+#endif //WITH_EDITOR
+
+		UE::DatasmithInterchange::FDatasmithReferenceMaterialManager::Destroy();
+	}
+
+private:
+#if WITH_EDITOR
+	void ExtendDatasmitMenuOptions(FToolMenuSection& SubSection);
+
+	void OnImportInterchange();
+
+	bool Import(const FString& FilePath, const FString& ContentPath);
+#endif //WITH_EDITOR
+};
+
+#if WITH_EDITOR
+void FDatasmithInterchangeModule::ExtendDatasmitMenuOptions(FToolMenuSection& SubSection)
+{
+	SubSection.AddMenuEntry(
+		TEXT("InterchangeImportFile"),
+		LOCTEXT("DatasmithInterchangeFileImport", "Interchange Import..."), // label
+		LOCTEXT("DatasmithInterchangeFileImportTooltip", "Experimental: Import Unreal Datasmith file using Interchange"), // description
+		FSlateIcon(),
+		FUIAction(FExecuteAction::CreateRaw(this, &FDatasmithInterchangeModule::OnImportInterchange))
+	);
+}
+
+void FDatasmithInterchangeModule::OnImportInterchange()
+{
+	const TArray<FString> Formats = FDatasmithTranslatorManager::Get().GetSupportedFormats();
+	FString FileTypes;
+	FString Extensions;
+
+	FString DefaultLocation(FEditorDirectories::Get().GetLastDirectory(ELastDirectory::GENERIC_IMPORT));
+
+	ObjectTools::AppendFormatsFileExtensions(Formats, FileTypes, Extensions);
+
+	const FString FormatString = FString::Printf(TEXT("All Files (%s)|%s|%s"), *Extensions, *Extensions, *FileTypes);
+	const FString Title = LOCTEXT("BrowseSourceDialogTitle", "Select Datasmith Source File").ToString();
+
+	TArray<FString> OutOpenFilenames;
+	IDesktopPlatform* DesktopPlatform = FDesktopPlatformModule::Get();
+	bool bOpened = false;
+	if (DesktopPlatform)
+	{
+		bOpened = DesktopPlatform->OpenFileDialog(
+			FSlateApplication::Get().FindBestParentWindowHandleForDialogs(nullptr),
+			Title,
+			DefaultLocation,
+			FString(),
+			FormatString,
+			EFileDialogFlags::None,
+			OutOpenFilenames
+		);
+	}
+
+	if (bOpened && OutOpenFilenames.Num() > 0)
+	{
+		static FString ContentPath = "/Game/"; // Trailing '/' is needed to set the default path
+
+		TSharedRef<SDlgPickPath> PickContentPathDlg =
+			SNew(SDlgPickPath)
+			.Title(LOCTEXT("ChooseImportRootContentPath", "Choose Location for importing the Datasmith content"))
+			.DefaultPath(FText::FromString(ContentPath));
+
+		if (PickContentPathDlg->ShowModal() != EAppReturnType::Cancel)
+		{
+			FEditorDirectories::Get().SetLastDirectory(ELastDirectory::GENERIC_IMPORT, FPaths::GetPath(OutOpenFilenames[0])); // Save path as default for next time.
+
+			ContentPath = PickContentPathDlg->GetPath().ToString() + "/";
+			Import(OutOpenFilenames[0], ContentPath);
+		}
+	}
+}
+
+bool FDatasmithInterchangeModule::Import(const FString& FilePath, const FString& ContentPath)
+{
+	UE::Interchange::FScopedSourceData ScopedSourceData(FilePath);
+
+	if (!UInterchangeManager::GetInterchangeManager().CanTranslateSourceData(ScopedSourceData.GetSourceData()))
+	{
+		return false;
+	}
+
+	FImportAssetParameters ImportAssetParameters;
+	ImportAssetParameters.bIsAutomated = false;
+
+	UInterchangeDatasmithPipeline* DatasmithPipeline = NewObject< UInterchangeDatasmithPipeline >();
+
+	ImportAssetParameters.OverridePipelines.Add(DatasmithPipeline);
+
+	TTuple<UE::Interchange::FAssetImportResultRef, UE::Interchange::FSceneImportResultRef> ImportResults
+		= UInterchangeManager::GetInterchangeManager().ImportSceneAsync(ContentPath, ScopedSourceData.GetSourceData(), ImportAssetParameters);
+	return ImportResults.Key->IsValid() && ImportResults.Value->IsValid();
+}
+#endif //WITH_EDITOR
+
+IMPLEMENT_MODULE(FDatasmithInterchangeModule, DatasmithInterchange);
+
+#undef LOCTEXT_NAMESPACE
\ No newline at end of file
diff --git a/Engine/Plugins/Experimental/Enterprise/DatasmithInterchange/Source/DatasmithInterchange/Private/InterchangeDatasmithAnimationUtils.cpp b/Engine/Plugins/Experimental/Enterprise/DatasmithInterchange/Source/DatasmithInterchange/Private/InterchangeDatasmithAnimationUtils.cpp
new file mode 100644
--- /dev/null
+++ b/Engine/Plugins/Experimental/Enterprise/DatasmithInterchange/Source/DatasmithInterchange/Private/InterchangeDatasmithAnimationUtils.cpp
@@ -0,0 +1,426 @@
+// Copyright Epic Games, Inc. All Rights Reserved.
+
+#include "InterchangeDatasmithLog.h"
+#include "InterchangeDatasmithUtils.h"
+
+#include "DatasmithAnimationElements.h"
+
+#include "Animation/InterchangeAnimationPayload.h"
+#include "InterchangeAnimationTrackSetNode.h"
+#include "InterchangeAnimSequenceFactoryNode.h"
+
+#include "Curves/RichCurve.h"
+#include "Misc/FrameRate.h"
+
+#define LOCTEXT_NAMESPACE "InterchangeDatasmithAnimationUtils"
+
+namespace UE::DatasmithInterchange::AnimUtils
+{
+	class FLevelSequenceImportHelper
+	{
+	public:
+		using FLevelSequencePtr = TSharedPtr<IDatasmithLevelSequenceElement>;
+
+		FLevelSequenceImportHelper(UInterchangeBaseNodeContainer& InBaseNodeContainer)
+			: BaseNodeContainer(InBaseNodeContainer)
+		{}
+
+		void Sort(TArray<FLevelSequencePtr>& LevelSequences);
+		void ProcessLevelSequence(FLevelSequencePtr& InLevelSequenceElement, UInterchangeAnimationTrackSetNode* InLevelSequenceNode);
+
+		TMap<FString, UE::DatasmithInterchange::AnimUtils::FAnimationPayloadDesc> PayLoadsMap;
+
+	private:
+		void ProcessSubsequenceTrack(IDatasmithSubsequenceAnimationElement& SubsequenceAnimation);
+		FString ProcessVisibilityTrack(IDatasmithVisibilityAnimationElement& VisibilityAnimation, int32 TrackIndex);
+		FString ProcessTransformTrack(IDatasmithTransformAnimationElement& TransformAnimation, int32 TrackIndex);
+
+		void AddLevelSequence(FLevelSequencePtr& LevelSequence, TArray<FLevelSequencePtr>& SortedLevelSequences, TSet<FLevelSequencePtr>& Processed);
+
+		template<class T>
+		T* AddTrackNode(int32 TrackIndex)
+		{
+			if (T* TrackNode = NewObject<T>(&BaseNodeContainer))
+			{
+				const FString TrackName = LevelSequenceNode->GetDisplayLabel() + TEXT("_Track_") + FString::FromInt(TrackIndex);
+				const FString TrackNodeUid = NodeUtils::LevelSequencePrefix + TrackName;
+				TrackNode->InitializeNode(TrackNodeUid, TrackName, EInterchangeNodeContainerType::TranslatedScene);
+
+				BaseNodeContainer.AddNode(TrackNode);
+
+				return TrackNode;
+			}
+
+			return nullptr;
+		}
+
+	private:
+		UInterchangeBaseNodeContainer& BaseNodeContainer;
+		FLevelSequencePtr LevelSequenceElement;
+		UInterchangeAnimationTrackSetNode* LevelSequenceNode = nullptr;
+		TMap<FLevelSequencePtr, FString> TranslatedLevelSequences;
+	};
+
+	void FLevelSequenceImportHelper::ProcessLevelSequence(FLevelSequencePtr& InLevelSequenceElement, UInterchangeAnimationTrackSetNode* InLevelSequenceNode)
+	{
+		LevelSequenceElement = InLevelSequenceElement;
+		LevelSequenceNode = InLevelSequenceNode;
+
+		const float FrameRate = LevelSequenceElement->GetFrameRate();
+		const int32 NumAnimations = LevelSequenceElement->GetAnimationsCount();
+
+		for (int32 TrackIndex = 0; TrackIndex < NumAnimations; ++TrackIndex)
+		{
+			TSharedPtr<IDatasmithBaseAnimationElement> AnimationElement = LevelSequenceElement->GetAnimation(TrackIndex);
+			if (!AnimationElement)
+			{
+				continue;
+			}
+
+			FString PayLoadKey;
+
+			if (AnimationElement->IsSubType(EDatasmithElementAnimationSubType::TransformAnimation))
+			{
+				PayLoadKey = ProcessTransformTrack(static_cast<IDatasmithTransformAnimationElement&>(*AnimationElement), TrackIndex);
+			}
+			else if (AnimationElement->IsSubType(EDatasmithElementAnimationSubType::VisibilityAnimation))
+			{
+				PayLoadKey = ProcessVisibilityTrack(static_cast<IDatasmithVisibilityAnimationElement&>(*AnimationElement), TrackIndex);
+			}
+			else if (AnimationElement->IsSubType(EDatasmithElementAnimationSubType::SubsequenceAnimation))
+			{
+				ProcessSubsequenceTrack(static_cast<IDatasmithSubsequenceAnimationElement&>(*AnimationElement));
+			}
+
+			if (!PayLoadKey.IsEmpty())
+			{
+				PayLoadsMap.Add(PayLoadKey, {FrameRate, AnimationElement});
+			}
+		}
+	}
+
+	void FLevelSequenceImportHelper::ProcessSubsequenceTrack(IDatasmithSubsequenceAnimationElement& SubsequenceAnimation)
+	{
+		TSharedPtr<IDatasmithLevelSequenceElement> PinnedTargetSequence = SubsequenceAnimation.GetSubsequence().Pin();
+		if (!ensure(PinnedTargetSequence))
+		{
+			UE_LOG(LogInterchangeDatasmith, Warning, TEXT("Cannot find LevelSequence element for Subsequence element %s"), SubsequenceAnimation.GetLabel());
+			return;
+		}
+
+		UInterchangeAnimationTrackSetInstanceNode* TrackNode = NewObject< UInterchangeAnimationTrackSetInstanceNode >(&BaseNodeContainer);
+		const FString TrackName = LevelSequenceNode->GetDisplayLabel() + TEXT("_SubTrack");
+		const FString TrackNodeUid = NodeUtils::LevelSequencePrefix + TrackName;
+		TrackNode->InitializeNode(TrackNodeUid, TrackName, EInterchangeNodeContainerType::TranslatedScene);
+
+		TrackNode->SetCustomStartFrame(SubsequenceAnimation.GetStartTime().Value);
+		TrackNode->SetCustomDuration(SubsequenceAnimation.GetDuration());
+		TrackNode->SetCustomTimeScale(SubsequenceAnimation.GetTimeScale());
+
+		const FString SequenceNodeUid = NodeUtils::GetLevelSequenceUid(PinnedTargetSequence->GetName());
+		TrackNode->SetCustomTrackSetDependencyUid(SequenceNodeUid);
+
+		BaseNodeContainer.AddNode(TrackNode);
+
+		LevelSequenceNode->AddCustomAnimationTrackUid(TrackNodeUid);
+	}
+
+	FString FLevelSequenceImportHelper::ProcessVisibilityTrack(IDatasmithVisibilityAnimationElement& VisibilityAnimation, int32 TrackIndex)
+	{
+		UInterchangeAnimationTrackNode* TrackNode = AddTrackNode<UInterchangeAnimationTrackNode>(TrackIndex);
+		if (!ensure(TrackNode))
+		{
+			UE_LOG(LogInterchangeDatasmith, Warning, TEXT("Cannot create UInterchangeAnimationTrackNode object for Visibility element %s"), VisibilityAnimation.GetLabel());
+			return {};
+		}
+
+		const FString ActorNodeUid = NodeUtils::GetActorUid(VisibilityAnimation.GetName());
+		TrackNode->SetCustomActorDependencyUid(ActorNodeUid);
+
+		TrackNode->SetCustomTargetedProperty((int32)EInterchangeAnimatedProperty::Visibility);
+		TrackNode->SetCustomAnimationPayloadKey(TrackNode->GetUniqueID());
+
+		LevelSequenceNode->AddCustomAnimationTrackUid(TrackNode->GetUniqueID());
+
+		return TrackNode->GetUniqueID();
+	}
+
+	FString FLevelSequenceImportHelper::ProcessTransformTrack(IDatasmithTransformAnimationElement& TransformAnimation, int32 TrackIndex)
+	{
+		UInterchangeTransformAnimationTrackNode* TrackNode = AddTrackNode< UInterchangeTransformAnimationTrackNode>(TrackIndex);
+		if (!ensure(TrackNode))
+		{
+			UE_LOG(LogInterchangeDatasmith, Warning, TEXT("Cannot create UInterchangeTransformAnimationTrackNode object for Visibility element %s"), TransformAnimation.GetLabel());
+			return {};
+		}
+
+		const FString ActorNodeUid = NodeUtils::GetActorUid(TransformAnimation.GetName());
+		TrackNode->SetCustomActorDependencyUid(ActorNodeUid);
+
+		TrackNode->SetCustomCompletionMode((int32)TransformAnimation.GetCompletionMode());
+		TrackNode->SetCustomUsedChannels((int32)TransformAnimation.GetEnabledTransformChannels());
+		TrackNode->SetCustomAnimationPayloadKey(TrackNode->GetUniqueID());
+
+		LevelSequenceNode->AddCustomAnimationTrackUid(TrackNode->GetUniqueID());
+
+		return TrackNode->GetUniqueID();
+	}
+
+	void FLevelSequenceImportHelper::AddLevelSequence(FLevelSequencePtr& LevelSequence, TArray<FLevelSequencePtr>& Sorted, TSet<FLevelSequencePtr>& Processed)
+	{
+		if (Processed.Contains(LevelSequence))
+		{
+			return;
+		}
+
+		const int32 NumAnimations = LevelSequence->GetAnimationsCount();
+		for (int32 TrackIndex = 0; TrackIndex < NumAnimations; ++TrackIndex)
+		{
+			TSharedPtr<IDatasmithBaseAnimationElement> AnimationElement = LevelSequence->GetAnimation(TrackIndex);
+			if (!AnimationElement)
+			{
+				continue;
+			}
+
+			if (AnimationElement->IsSubType(EDatasmithElementAnimationSubType::SubsequenceAnimation))
+			{
+				IDatasmithSubsequenceAnimationElement& SubsequenceAnimation = static_cast<IDatasmithSubsequenceAnimationElement&>(*AnimationElement);
+				TSharedPtr<IDatasmithLevelSequenceElement> PinnedTargetSequence = SubsequenceAnimation.GetSubsequence().Pin();
+				if (PinnedTargetSequence && !Processed.Contains(PinnedTargetSequence))
+				{
+					AddLevelSequence(PinnedTargetSequence, Sorted, Processed);
+				}
+			}
+		}
+
+		Sorted.Add(LevelSequence);
+		Processed.Add(LevelSequence);
+	}
+
+	void FLevelSequenceImportHelper::Sort(TArray<FLevelSequencePtr>& LevelSequences)
+	{
+		TSet<FLevelSequencePtr> Processed;
+		TArray<FLevelSequencePtr> Sorted;
+		Sorted.Reserve(LevelSequences.Num());
+
+		for (TSharedPtr<IDatasmithLevelSequenceElement>& LevelSequence : LevelSequences)
+		{
+			AddLevelSequence(LevelSequence, Sorted, Processed);
+		}
+
+		LevelSequences.Reset();
+		LevelSequences = MoveTemp(Sorted);
+	}
+
+	class FAnimationCurvesHelper
+	{
+	public:
+		FAnimationCurvesHelper(const IDatasmithTransformAnimationElement& InAnimation, const FFrameRate& InSrcFrameRate, TArray<FRichCurve>& InCurves)
+			: Animation(InAnimation)
+			, SrcFrameRate(InSrcFrameRate)
+			, Curves(InCurves)
+		{
+			ensure(Curves.Num() == 9);
+		}
+
+		bool FillCurves(EDatasmithTransformType TransformType)
+		{
+			bool bIsValidAnimation = false;
+			const int32 NumFrames = Animation.GetFramesCount(TransformType);
+
+			for (int32 FrameIndex = 0; FrameIndex < NumFrames; ++FrameIndex)
+			{
+				if (Animation.GetFrame(TransformType, FrameIndex).IsValid())
+				{
+					bIsValidAnimation = true;
+					break;
+				}
+			}
+
+			if (!bIsValidAnimation)
+			{
+				return false;
+			}
+
+			static_assert((int)ERichCurveInterpMode::RCIM_Linear == (int)EDatasmithCurveInterpMode::Linear, "INVALID_ENUM_VALUE");
+			static_assert((int)ERichCurveInterpMode::RCIM_Constant == (int)EDatasmithCurveInterpMode::Constant, "INVALID_ENUM_VALUE");
+			static_assert((int)ERichCurveInterpMode::RCIM_Cubic == (int)EDatasmithCurveInterpMode::Cubic, "INVALID_ENUM_VALUE");
+
+			const ERichCurveInterpMode Interpolation = (ERichCurveInterpMode)Animation.GetCurveInterpMode(TransformType);
+
+			TFunction<FKeyHandle(FRichCurve&, float, float)> AddKey = [Interpolation](FRichCurve& Curve, float KeyTime, float KeyValue) -> FKeyHandle
+			{
+				FKeyHandle KeyHandle = Curve.AddKey(KeyTime, KeyValue);
+				Curve.SetKeyInterpMode(KeyHandle, Interpolation);
+
+				return KeyHandle;
+			};
+
+			const int32 CurveOffset = TransformType == EDatasmithTransformType::Translation ? 0 : (TransformType == EDatasmithTransformType::Rotation ? 3 : 6);
+
+			FRichCurve& XCurve = Curves[CurveOffset + 0];
+			FRichCurve& YCurve = Curves[CurveOffset + 1];
+			FRichCurve& ZCurve = Curves[CurveOffset + 2];
+
+
+			//  Rewind the rotations to prevent rotations that look like axis flips (eg. from 180 to -180 degrees)
+			if (EDatasmithTransformType::Rotation == TransformType)
+			{
+				FKeyHandle XPrevKeyHandle;
+				FKeyHandle YPrevKeyHandle;
+				FKeyHandle ZPrevKeyHandle;
+				int32 StartIndex = 1;
+
+				for (int32 FrameIndex = 0; FrameIndex < NumFrames; ++FrameIndex)
+				{
+					const FDatasmithTransformFrameInfo& FrameInfo = Animation.GetFrame(TransformType, FrameIndex);
+					if (FrameInfo.IsValid())
+					{
+						const float FrameTime = SrcFrameRate.AsSeconds(FrameInfo.FrameNumber);
+
+						XPrevKeyHandle = AddKey(XCurve, FrameTime, FrameInfo.X);
+						YPrevKeyHandle = AddKey(YCurve, FrameTime, FrameInfo.Y);
+						ZPrevKeyHandle = AddKey(ZCurve, FrameTime, FrameInfo.Z);
+
+						break;
+					}
+
+					++StartIndex;
+				}
+
+				for (int32 FrameIndex = StartIndex; FrameIndex < NumFrames; ++FrameIndex)
+				{
+					const FDatasmithTransformFrameInfo& FrameInfo = Animation.GetFrame(TransformType, FrameIndex);
+					if (FrameInfo.IsValid())
+					{
+						const float FrameTime = SrcFrameRate.AsSeconds(FrameInfo.FrameNumber);
+
+						float KeyValue = FrameInfo.X;
+						FMath::WindRelativeAnglesDegrees(XCurve.GetKeyValue(XPrevKeyHandle), KeyValue);
+						XPrevKeyHandle = AddKey(XCurve, FrameTime, KeyValue);
+
+						KeyValue = FrameInfo.Y;
+						FMath::WindRelativeAnglesDegrees(YCurve.GetKeyValue(YPrevKeyHandle), KeyValue);
+						YPrevKeyHandle = AddKey(YCurve, FrameTime, KeyValue);
+
+						KeyValue = FrameInfo.Z;
+						FMath::WindRelativeAnglesDegrees(ZCurve.GetKeyValue(ZPrevKeyHandle), KeyValue);
+						ZPrevKeyHandle = AddKey(ZCurve, FrameTime, KeyValue);
+					}
+				}
+			}
+			else
+			{
+				for (int32 FrameIndex = 0; FrameIndex < NumFrames; ++FrameIndex)
+				{
+					const FDatasmithTransformFrameInfo& FrameInfo = Animation.GetFrame(TransformType, FrameIndex);
+					if (FrameInfo.IsValid())
+					{
+						const float FrameTime = SrcFrameRate.AsSeconds(FrameInfo.FrameNumber);
+
+						AddKey(XCurve, FrameTime, FrameInfo.X);
+						AddKey(YCurve, FrameTime, FrameInfo.Y);
+						AddKey(ZCurve, FrameTime, FrameInfo.Z);
+					}
+				}
+			}
+
+			return true;
+		}
+
+	private:
+		const IDatasmithTransformAnimationElement& Animation;
+		const FFrameRate& SrcFrameRate;
+		TArray<FRichCurve>& Curves;
+	};
+
+	bool GetAnimationPayloadData(const IDatasmithBaseAnimationElement& AnimationElement, float FrameRate, UE::Interchange::FAnimationCurvePayloadData& PayLoadData)
+	{
+		using namespace UE::Interchange::Animation;
+
+		if (AnimationElement.IsSubType(EDatasmithElementAnimationSubType::TransformAnimation))
+		{
+			const IDatasmithTransformAnimationElement& TransformAnimation = static_cast<const IDatasmithTransformAnimationElement&>(AnimationElement);
+
+			TArray<FRichCurve>& Curves = PayLoadData.Curves;
+			Curves.SetNum(9);
+
+			FAnimationCurvesHelper Helper(TransformAnimation, ConvertSampleRatetoFrameRate(FrameRate), Curves);
+
+			Helper.FillCurves(EDatasmithTransformType::Translation);
+			Helper.FillCurves(EDatasmithTransformType::Rotation);
+			Helper.FillCurves(EDatasmithTransformType::Scale);
+
+			return true;
+		}
+
+		return false;
+	}
+
+	bool GetAnimationPayloadData(const IDatasmithBaseAnimationElement& AnimationElement, float FloatFrameRate, UE::Interchange::FAnimationStepCurvePayloadData& PayLoadData)
+	{
+		using namespace UE::Interchange::Animation;
+
+		if (AnimationElement.IsSubType(EDatasmithElementAnimationSubType::VisibilityAnimation))
+		{
+			const IDatasmithVisibilityAnimationElement& VisibilityAnimation = static_cast<const IDatasmithVisibilityAnimationElement&>(AnimationElement);
+			const FFrameRate FrameRate = ConvertSampleRatetoFrameRate(FloatFrameRate);
+
+			FInterchangeStepCurve& Curve = PayLoadData.StepCurves.AddDefaulted_GetRef();
+
+			const int32 NumFrames = VisibilityAnimation.GetFramesCount();
+
+			TArray<float>& KeyTimes = Curve.KeyTimes;
+			KeyTimes.Reserve(NumFrames);
+
+			TArray<bool>& BooleanValues = Curve.BooleanKeyValues.Emplace();
+			BooleanValues.Reserve(NumFrames);
+
+			for (int32 FrameIndex = 0; FrameIndex < NumFrames; ++FrameIndex)
+			{
+				const FDatasmithVisibilityFrameInfo& FrameInfo = VisibilityAnimation.GetFrame(FrameIndex);
+				if (FrameInfo.IsValid())
+				{
+					KeyTimes.Add(FrameRate.AsSeconds(FrameInfo.FrameNumber));
+					BooleanValues.Add(FrameInfo.bVisible);
+				}
+			}
+
+			return true;
+		}
+
+		return false;
+	}
+
+	void TranslateLevelSequences(TArray<TSharedPtr<IDatasmithLevelSequenceElement>>& LevelSequences, UInterchangeBaseNodeContainer& BaseNodeContainer, TMap<FString, UE::DatasmithInterchange::AnimUtils::FAnimationPayloadDesc>& AnimationPayLoadMapping)
+	{
+		FLevelSequenceImportHelper LevelSequenceHelper(BaseNodeContainer);
+
+		// Sort the array from less dependent to more dependent
+		LevelSequenceHelper.Sort(LevelSequences);
+
+		for (TSharedPtr<IDatasmithLevelSequenceElement>& SequenceElement : LevelSequences)
+		{
+			if (!SequenceElement || SequenceElement->GetAnimationsCount() == 0)
+			{
+				continue;
+			}
+
+			UInterchangeAnimationTrackSetNode* SequenceNode = NewObject< UInterchangeAnimationTrackSetNode >(&BaseNodeContainer);
+			if (ensure(SequenceNode))
+			{
+				const FString SequenceName = SequenceElement->GetName();
+				const FString SequenceNodeUid = NodeUtils::GetLevelSequenceUid(*SequenceName);
+				SequenceNode->InitializeNode(SequenceNodeUid, SequenceElement->GetLabel(), EInterchangeNodeContainerType::TranslatedScene);
+
+				BaseNodeContainer.AddNode(SequenceNode);
+
+				LevelSequenceHelper.ProcessLevelSequence(SequenceElement, SequenceNode);
+			}
+		}
+
+		AnimationPayLoadMapping = MoveTemp(LevelSequenceHelper.PayLoadsMap);
+	}
+}
+
+#undef LOCTEXT_NAMESPACE
\ No newline at end of file
diff --git a/Engine/Plugins/Experimental/Enterprise/DatasmithInterchange/Source/DatasmithInterchange/Private/InterchangeDatasmithAreaLightFactory.cpp b/Engine/Plugins/Experimental/Enterprise/DatasmithInterchange/Source/DatasmithInterchange/Private/InterchangeDatasmithAreaLightFactory.cpp
new file mode 100644
--- /dev/null
+++ b/Engine/Plugins/Experimental/Enterprise/DatasmithInterchange/Source/DatasmithInterchange/Private/InterchangeDatasmithAreaLightFactory.cpp
@@ -0,0 +1,86 @@
+// Copyright Epic Games, Inc. All Rights Reserved.
+
+#include "InterchangeDatasmithAreaLightFactory.h"
+
+#include "InterchangeDatasmithLog.h"
+#include "InterchangeDatasmithAreaLightFactoryNode.h"
+
+#include "Scene/InterchangeActorHelper.h"
+#include "DatasmithAreaLightActor.h"
+#include "Engine/TextureLightProfile.h"
+#include "InterchangeTextureFactoryNode.h"
+#include "Logging/LogMacros.h"
+#include "Nodes/InterchangeBaseNode.h"
+#include "Nodes/InterchangeBaseNodeContainer.h"
+
+UClass* UInterchangeDatasmithAreaLightFactory::GetFactoryClass() const
+{
+	FSoftObjectPath LightShapeBlueprintRef = FSoftObjectPath(TEXT("/DatasmithContent/Datasmith/DatasmithArealight.DatasmithArealight"));
+	UBlueprint* LightShapeBlueprint = Cast< UBlueprint >(LightShapeBlueprintRef.ResolveObject());
+
+	// LightShapeBlueprint should already be loaded at this point.
+	if (!ensure(LightShapeBlueprint))
+	{
+		return nullptr;
+	}
+
+	return LightShapeBlueprint->GeneratedClass;
+}
+
+UObject* UInterchangeDatasmithAreaLightFactory::CreateSceneObject(const UInterchangeFactoryBase::FCreateSceneObjectsParams& CreateSceneObjectsParams)
+{
+	//using namespace UE::Interchange::ActorHelper;
+
+	UInterchangeDatasmithAreaLightFactoryNode* AreaLightFactoryNode = Cast<UInterchangeDatasmithAreaLightFactoryNode>(CreateSceneObjectsParams.FactoryNode);
+	if (!AreaLightFactoryNode)
+	{
+		return nullptr;
+	}
+
+	ADatasmithAreaLightActor* SpawnedAreaLightActor = Cast<ADatasmithAreaLightActor>(UE::Interchange::ActorHelper::SpawnFactoryActor(CreateSceneObjectsParams));
+	if (!SpawnedAreaLightActor)
+	{
+		return nullptr;
+	}
+
+	// Update AreaLight properties and rebuild the actor.
+	{
+		SpawnedAreaLightActor->UnregisterAllComponents(true);
+
+		// Find the IES Texture asset and apply it.
+		FString IESTextureUid;
+		if (AreaLightFactoryNode->GetCustomIESTexture(IESTextureUid))
+		{
+			const UInterchangeBaseNodeContainer* NodeContainer = CreateSceneObjectsParams.NodeContainer;
+
+			if (const UInterchangeBaseNode* IESTextureNode = NodeContainer->GetNode(IESTextureUid))
+			{
+				TArray<FString> TargetNodesUid;
+				IESTextureNode->GetTargetNodeUids(TargetNodesUid);
+				for (const FString& TargetNodeUid : TargetNodesUid)
+				{
+					if (UInterchangeTextureFactoryNode* TextureFactoryNode = Cast<UInterchangeTextureFactoryNode>(NodeContainer->GetFactoryNode(TargetNodeUid)))
+					{
+						FSoftObjectPath TextureReferenceObject;
+						TextureFactoryNode->GetCustomReferenceObject(TextureReferenceObject);
+						if (UTextureLightProfile* TextureLightProfile = Cast<UTextureLightProfile>(TextureReferenceObject.TryLoad()))
+						{
+							SpawnedAreaLightActor->IESTexture = TextureLightProfile;
+							break;
+						}
+					}
+				}
+			}
+		}
+
+		// Apply all simple attributes.
+		AreaLightFactoryNode->ApplyAllCustomAttributeToObject(SpawnedAreaLightActor);
+
+		SpawnedAreaLightActor->RegisterAllComponents();
+#if WITH_EDITOR
+		SpawnedAreaLightActor->RerunConstructionScripts();
+#endif //WITH_EDITOR
+	}
+
+	return SpawnedAreaLightActor;
+};
\ No newline at end of file
diff --git a/Engine/Plugins/Experimental/Enterprise/DatasmithInterchange/Source/DatasmithInterchange/Private/InterchangeDatasmithAreaLightFactoryNode.cpp b/Engine/Plugins/Experimental/Enterprise/DatasmithInterchange/Source/DatasmithInterchange/Private/InterchangeDatasmithAreaLightFactoryNode.cpp
new file mode 100644
--- /dev/null
+++ b/Engine/Plugins/Experimental/Enterprise/DatasmithInterchange/Source/DatasmithInterchange/Private/InterchangeDatasmithAreaLightFactoryNode.cpp
@@ -0,0 +1,215 @@
+// Copyright Epic Games, Inc. All Rights Reserved.
+
+#include "InterchangeDatasmithAreaLightFactoryNode.h"
+
+#include "DatasmithAreaLightActor.h"
+#include "Engine/Blueprint.h"
+
+UE::Interchange::FAttributeKey UInterchangeDatasmithAreaLightFactoryNode::Macro_CustomLightTypeKey = UE::Interchange::FAttributeKey(TEXT("LightType"));
+UE::Interchange::FAttributeKey UInterchangeDatasmithAreaLightFactoryNode::Macro_CustomLightShapeKey = UE::Interchange::FAttributeKey(TEXT("LightShape"));
+UE::Interchange::FAttributeKey UInterchangeDatasmithAreaLightFactoryNode::Macro_CustomDimensionsKey = UE::Interchange::FAttributeKey(TEXT("Dimensions"));
+UE::Interchange::FAttributeKey UInterchangeDatasmithAreaLightFactoryNode::Macro_CustomIntensityKey = UE::Interchange::FAttributeKey(TEXT("Intensity"));
+UE::Interchange::FAttributeKey UInterchangeDatasmithAreaLightFactoryNode::Macro_CustomIntensityUnitsKey = UE::Interchange::FAttributeKey(TEXT("IntensityUnits"));
+UE::Interchange::FAttributeKey UInterchangeDatasmithAreaLightFactoryNode::Macro_CustomColorKey = UE::Interchange::FAttributeKey(TEXT("Color"));
+UE::Interchange::FAttributeKey UInterchangeDatasmithAreaLightFactoryNode::Macro_CustomTemperatureKey = UE::Interchange::FAttributeKey(TEXT("Temperature"));
+UE::Interchange::FAttributeKey UInterchangeDatasmithAreaLightFactoryNode::Macro_CustomIESTextureKey = UE::Interchange::FAttributeKey(TEXT("IESTexture"));
+UE::Interchange::FAttributeKey UInterchangeDatasmithAreaLightFactoryNode::Macro_CustomUseIESBrightnessKey = UE::Interchange::FAttributeKey(TEXT("UseIESBrightness"));
+UE::Interchange::FAttributeKey UInterchangeDatasmithAreaLightFactoryNode::Macro_CustomIESBrightnessScaleKey = UE::Interchange::FAttributeKey(TEXT("IESBrightnessScale"));
+UE::Interchange::FAttributeKey UInterchangeDatasmithAreaLightFactoryNode::Macro_CustomRotationKey = UE::Interchange::FAttributeKey(TEXT("Rotation"));
+UE::Interchange::FAttributeKey UInterchangeDatasmithAreaLightFactoryNode::Macro_CustomSourceRadiusKey = UE::Interchange::FAttributeKey(TEXT("SourceRadius"));
+UE::Interchange::FAttributeKey UInterchangeDatasmithAreaLightFactoryNode::Macro_CustomSourceLengthKey = UE::Interchange::FAttributeKey(TEXT("SourceLength"));
+UE::Interchange::FAttributeKey UInterchangeDatasmithAreaLightFactoryNode::Macro_CustomAttenuationRadiusKey = UE::Interchange::FAttributeKey(TEXT("AttenuationRadius"));
+UE::Interchange::FAttributeKey UInterchangeDatasmithAreaLightFactoryNode::Macro_CustomSpotlightInnerAngleKey = UE::Interchange::FAttributeKey(TEXT("SpotlightInnerAngle"));
+UE::Interchange::FAttributeKey UInterchangeDatasmithAreaLightFactoryNode::Macro_CustomSpotlightOuterAngleKey = UE::Interchange::FAttributeKey(TEXT("SpotlightOuterAngle"));
+
+FString UInterchangeDatasmithAreaLightFactoryNode::GetTypeName() const
+{
+	const FString TypeName = TEXT("DatasmithAreaLightFactoryNode");
+	return TypeName;
+}
+
+UClass* UInterchangeDatasmithAreaLightFactoryNode::GetObjectClass() const
+{
+	FSoftObjectPath LightShapeBlueprintRef = FSoftObjectPath(TEXT("/DatasmithContent/Datasmith/DatasmithArealight.DatasmithArealight"));
+	UBlueprint* LightShapeBlueprint = Cast<UBlueprint>(LightShapeBlueprintRef.ResolveObject());
+
+	// LightShapeBlueprint should be loaded at this point.
+	if (!ensure(LightShapeBlueprint))
+	{
+		return nullptr;
+	}
+
+	return LightShapeBlueprint->GeneratedClass;
+}
+
+bool UInterchangeDatasmithAreaLightFactoryNode::GetCustomLightType(EDatasmithAreaLightActorType& OutAttributeValue) const
+{
+	uint8 AttributeValue = 0;
+	bool bSuccess = [this, &AttributeValue]() { IMPLEMENT_NODE_ATTRIBUTE_GETTER(LightType, uint8); }();
+	OutAttributeValue = static_cast<EDatasmithAreaLightActorType>(AttributeValue);
+	return bSuccess;
+}
+
+bool UInterchangeDatasmithAreaLightFactoryNode::SetCustomLightType(const EDatasmithAreaLightActorType& InAttributeValue, bool bAddApplyDelegate)
+{
+	uint8 AttributeValue = static_cast<uint8>(InAttributeValue);
+	IMPLEMENT_NODE_ATTRIBUTE_SETTER(UInterchangeDatasmithAreaLightFactoryNode, LightType, uint8, ADatasmithAreaLightActor);
+}
+
+bool UInterchangeDatasmithAreaLightFactoryNode::GetCustomLightShape(EDatasmithAreaLightActorShape& OutAttributeValue) const
+{
+	uint8 AttributeValue = 0;
+	bool bSuccess = [this, &AttributeValue]() { IMPLEMENT_NODE_ATTRIBUTE_GETTER(LightShape, uint8); }();
+	OutAttributeValue = static_cast<EDatasmithAreaLightActorShape>(AttributeValue);
+	return bSuccess;
+}
+
+bool UInterchangeDatasmithAreaLightFactoryNode::SetCustomLightShape(const EDatasmithAreaLightActorShape& InAttributeValue, bool bAddApplyDelegate)
+{
+	uint8 AttributeValue = static_cast<uint8>(InAttributeValue);
+	IMPLEMENT_NODE_ATTRIBUTE_SETTER(UInterchangeDatasmithAreaLightFactoryNode, LightShape, uint8, ADatasmithAreaLightActor);
+}
+
+bool UInterchangeDatasmithAreaLightFactoryNode::GetCustomDimensions(FVector2D& AttributeValue) const
+{
+	IMPLEMENT_NODE_ATTRIBUTE_GETTER(Dimensions, FVector2D);
+}
+
+bool UInterchangeDatasmithAreaLightFactoryNode::SetCustomDimensions(const FVector2D& AttributeValue, bool bAddApplyDelegate)
+{
+	IMPLEMENT_NODE_ATTRIBUTE_SETTER(UInterchangeDatasmithAreaLightFactoryNode, Dimensions, FVector2D, ADatasmithAreaLightActor);
+}
+
+bool UInterchangeDatasmithAreaLightFactoryNode::GetCustomIntensity(float& AttributeValue) const
+{
+	IMPLEMENT_NODE_ATTRIBUTE_GETTER(Intensity, float);
+}
+
+bool UInterchangeDatasmithAreaLightFactoryNode::SetCustomIntensity(const float& AttributeValue, bool bAddApplyDelegate)
+{
+	IMPLEMENT_NODE_ATTRIBUTE_SETTER(UInterchangeDatasmithAreaLightFactoryNode, Intensity, float, ADatasmithAreaLightActor);
+}
+
+bool UInterchangeDatasmithAreaLightFactoryNode::GetCustomIntensityUnits(ELightUnits& OutAttributeValue) const
+{
+	uint8 AttributeValue = 0;
+	bool bSuccess = [this, &AttributeValue]() { IMPLEMENT_NODE_ATTRIBUTE_GETTER(IntensityUnits, uint8); }();
+	OutAttributeValue = static_cast<ELightUnits>(AttributeValue);
+	return bSuccess;
+}
+
+bool UInterchangeDatasmithAreaLightFactoryNode::SetCustomIntensityUnits(const ELightUnits& InAttributeValue, bool bAddApplyDelegate)
+{
+	uint8 AttributeValue = static_cast<uint8>(InAttributeValue);
+	IMPLEMENT_NODE_ATTRIBUTE_SETTER(UInterchangeDatasmithAreaLightFactoryNode, IntensityUnits, uint8, ADatasmithAreaLightActor);
+}
+
+bool UInterchangeDatasmithAreaLightFactoryNode::GetCustomColor(FLinearColor& AttributeValue) const
+{
+	IMPLEMENT_NODE_ATTRIBUTE_GETTER(Color, FLinearColor);
+}
+
+bool UInterchangeDatasmithAreaLightFactoryNode::SetCustomColor(const FLinearColor& AttributeValue, bool bAddApplyDelegate)
+{
+	IMPLEMENT_NODE_ATTRIBUTE_SETTER(UInterchangeDatasmithAreaLightFactoryNode, Color, FLinearColor, ADatasmithAreaLightActor);
+}
+
+bool UInterchangeDatasmithAreaLightFactoryNode::GetCustomTemperature(float& AttributeValue) const
+{
+	IMPLEMENT_NODE_ATTRIBUTE_GETTER(Temperature, float);
+}
+
+bool UInterchangeDatasmithAreaLightFactoryNode::SetCustomTemperature(const float& AttributeValue, bool bAddApplyDelegate)
+{
+	IMPLEMENT_NODE_ATTRIBUTE_SETTER(UInterchangeDatasmithAreaLightFactoryNode, Temperature, float, ADatasmithAreaLightActor);
+}
+
+bool UInterchangeDatasmithAreaLightFactoryNode::GetCustomIESTexture(FString& AttributeValue) const
+{
+	IMPLEMENT_NODE_ATTRIBUTE_GETTER(IESTexture, FString);
+}
+
+bool UInterchangeDatasmithAreaLightFactoryNode::SetCustomIESTexture(const FString& AttributeValue)
+{
+	IMPLEMENT_NODE_ATTRIBUTE_SETTER_NODELEGATE(IESTexture, FString);
+}
+
+bool UInterchangeDatasmithAreaLightFactoryNode::GetCustomUseIESBrightness(bool& AttributeValue) const
+{
+	IMPLEMENT_NODE_ATTRIBUTE_GETTER(UseIESBrightness, bool);
+}
+
+bool UInterchangeDatasmithAreaLightFactoryNode::SetCustomUseIESBrightness(const bool& AttributeValue, bool bAddApplyDelegate)
+{
+	IMPLEMENT_NODE_ATTRIBUTE_SETTER(UInterchangeDatasmithAreaLightFactoryNode, UseIESBrightness, bool, ADatasmithAreaLightActor);
+}
+
+bool UInterchangeDatasmithAreaLightFactoryNode::GetCustomIESBrightnessScale(float& AttributeValue) const
+{
+	IMPLEMENT_NODE_ATTRIBUTE_GETTER(IESBrightnessScale, float);
+}
+
+bool UInterchangeDatasmithAreaLightFactoryNode::SetCustomIESBrightnessScale(const float& AttributeValue, bool bAddApplyDelegate)
+{
+	IMPLEMENT_NODE_ATTRIBUTE_SETTER(UInterchangeDatasmithAreaLightFactoryNode, IESBrightnessScale, float, ADatasmithAreaLightActor);
+}
+
+bool UInterchangeDatasmithAreaLightFactoryNode::GetCustomRotation(FRotator& AttributeValue) const
+{
+	IMPLEMENT_NODE_ATTRIBUTE_GETTER(Rotation, FRotator);
+}
+
+bool UInterchangeDatasmithAreaLightFactoryNode::SetCustomRotation(const FRotator& AttributeValue, bool bAddApplyDelegate)
+{
+	IMPLEMENT_NODE_ATTRIBUTE_SETTER(UInterchangeDatasmithAreaLightFactoryNode, Rotation, FRotator, ADatasmithAreaLightActor);
+}
+
+bool UInterchangeDatasmithAreaLightFactoryNode::GetCustomSourceRadius(float& AttributeValue) const
+{
+	IMPLEMENT_NODE_ATTRIBUTE_GETTER(SourceRadius, float);
+}
+
+bool UInterchangeDatasmithAreaLightFactoryNode::SetCustomSourceRadius(const float& AttributeValue, bool bAddApplyDelegate)
+{
+	IMPLEMENT_NODE_ATTRIBUTE_SETTER(UInterchangeDatasmithAreaLightFactoryNode, SourceRadius, float, ADatasmithAreaLightActor);
+}
+
+bool UInterchangeDatasmithAreaLightFactoryNode::GetCustomSourceLength(float& AttributeValue) const
+{
+	IMPLEMENT_NODE_ATTRIBUTE_GETTER(SourceLength, float);
+}
+
+bool UInterchangeDatasmithAreaLightFactoryNode::SetCustomSourceLength(const float& AttributeValue, bool bAddApplyDelegate)
+{
+	IMPLEMENT_NODE_ATTRIBUTE_SETTER(UInterchangeDatasmithAreaLightFactoryNode, SourceLength, float, ADatasmithAreaLightActor);
+}
+
+bool UInterchangeDatasmithAreaLightFactoryNode::GetCustomAttenuationRadius(float& AttributeValue) const
+{
+	IMPLEMENT_NODE_ATTRIBUTE_GETTER(AttenuationRadius, float);
+}
+
+bool UInterchangeDatasmithAreaLightFactoryNode::SetCustomAttenuationRadius(const float& AttributeValue, bool bAddApplyDelegate)
+{
+	IMPLEMENT_NODE_ATTRIBUTE_SETTER(UInterchangeDatasmithAreaLightFactoryNode, AttenuationRadius, float, ADatasmithAreaLightActor);
+}
+
+bool UInterchangeDatasmithAreaLightFactoryNode::GetCustomSpotlightInnerAngle(float& AttributeValue) const
+{
+	IMPLEMENT_NODE_ATTRIBUTE_GETTER(SpotlightInnerAngle, float);
+}
+
+bool UInterchangeDatasmithAreaLightFactoryNode::SetCustomSpotlightInnerAngle(const float& AttributeValue, bool bAddApplyDelegate)
+{
+	IMPLEMENT_NODE_ATTRIBUTE_SETTER(UInterchangeDatasmithAreaLightFactoryNode, SpotlightInnerAngle, float, ADatasmithAreaLightActor);
+}
+
+bool UInterchangeDatasmithAreaLightFactoryNode::GetCustomSpotlightOuterAngle(float& AttributeValue) const
+{
+	IMPLEMENT_NODE_ATTRIBUTE_GETTER(SpotlightOuterAngle, float);
+}
+
+bool UInterchangeDatasmithAreaLightFactoryNode::SetCustomSpotlightOuterAngle(const float& AttributeValue, bool bAddApplyDelegate)
+{
+	IMPLEMENT_NODE_ATTRIBUTE_SETTER(UInterchangeDatasmithAreaLightFactoryNode, SpotlightOuterAngle, float, ADatasmithAreaLightActor);
+}
diff --git a/Engine/Plugins/Experimental/Enterprise/DatasmithInterchange/Source/DatasmithInterchange/Private/InterchangeDatasmithAreaLightNode.cpp b/Engine/Plugins/Experimental/Enterprise/DatasmithInterchange/Source/DatasmithInterchange/Private/InterchangeDatasmithAreaLightNode.cpp
new file mode 100644
--- /dev/null
+++ b/Engine/Plugins/Experimental/Enterprise/DatasmithInterchange/Source/DatasmithInterchange/Private/InterchangeDatasmithAreaLightNode.cpp
@@ -0,0 +1,146 @@
+// Copyright Epic Games, Inc. All Rights Reserved.
+
+#include "InterchangeDatasmithAreaLightNode.h"
+
+
+
+UE::Interchange::FAttributeKey UInterchangeDatasmithAreaLightNode::Macro_CustomLightTypeKey = UE::Interchange::FAttributeKey(TEXT("LightType"));
+UE::Interchange::FAttributeKey UInterchangeDatasmithAreaLightNode::Macro_CustomLightShapeKey = UE::Interchange::FAttributeKey(TEXT("LightShape"));
+UE::Interchange::FAttributeKey UInterchangeDatasmithAreaLightNode::Macro_CustomDimensionsKey = UE::Interchange::FAttributeKey(TEXT("Dimensions"));
+UE::Interchange::FAttributeKey UInterchangeDatasmithAreaLightNode::Macro_CustomColorKey = UE::Interchange::FAttributeKey(TEXT("Color"));
+UE::Interchange::FAttributeKey UInterchangeDatasmithAreaLightNode::Macro_CustomIESTextureKey = UE::Interchange::FAttributeKey(TEXT("IESTexture"));
+UE::Interchange::FAttributeKey UInterchangeDatasmithAreaLightNode::Macro_CustomUseIESBrightnessKey = UE::Interchange::FAttributeKey(TEXT("UseIESBrightness"));
+UE::Interchange::FAttributeKey UInterchangeDatasmithAreaLightNode::Macro_CustomIESBrightnessScaleKey = UE::Interchange::FAttributeKey(TEXT("IESBrightnessScale"));
+UE::Interchange::FAttributeKey UInterchangeDatasmithAreaLightNode::Macro_CustomRotationKey = UE::Interchange::FAttributeKey(TEXT("Rotation"));
+UE::Interchange::FAttributeKey UInterchangeDatasmithAreaLightNode::Macro_CustomSourceRadiusKey = UE::Interchange::FAttributeKey(TEXT("SourceRadius"));
+UE::Interchange::FAttributeKey UInterchangeDatasmithAreaLightNode::Macro_CustomSourceLengthKey = UE::Interchange::FAttributeKey(TEXT("SourceLength"));
+UE::Interchange::FAttributeKey UInterchangeDatasmithAreaLightNode::Macro_CustomSpotlightInnerAngleKey = UE::Interchange::FAttributeKey(TEXT("SpotlightInnerAngle"));
+UE::Interchange::FAttributeKey UInterchangeDatasmithAreaLightNode::Macro_CustomSpotlightOuterAngleKey = UE::Interchange::FAttributeKey(TEXT("SpotlightOuterAngle"));
+
+bool UInterchangeDatasmithAreaLightNode::GetCustomLightType(EDatasmithAreaLightActorType& OutAttributeValue) const
+{
+	uint8 AttributeValue = 0;
+	bool bSuccess = [this, &AttributeValue]() { IMPLEMENT_NODE_ATTRIBUTE_GETTER(LightType, uint8); }();
+	OutAttributeValue = static_cast<EDatasmithAreaLightActorType>(AttributeValue);
+	return bSuccess;
+}
+
+bool UInterchangeDatasmithAreaLightNode::SetCustomLightType(const EDatasmithAreaLightActorType& InAttributeValue, bool bAddApplyDelegate)
+{
+	uint8 AttributeValue = static_cast<uint8>(InAttributeValue);
+	IMPLEMENT_NODE_ATTRIBUTE_SETTER_NODELEGATE(LightType, uint8);
+}
+
+bool UInterchangeDatasmithAreaLightNode::GetCustomLightShape(EDatasmithAreaLightActorShape& OutAttributeValue) const
+{
+	uint8 AttributeValue = 0;
+	bool bSuccess = [this, &AttributeValue]() { IMPLEMENT_NODE_ATTRIBUTE_GETTER(LightShape, uint8); }();
+	OutAttributeValue = static_cast<EDatasmithAreaLightActorShape>(AttributeValue);
+	return bSuccess;
+}
+
+bool UInterchangeDatasmithAreaLightNode::SetCustomLightShape(const EDatasmithAreaLightActorShape& InAttributeValue, bool bAddApplyDelegate)
+{
+	uint8 AttributeValue = static_cast<uint8>(InAttributeValue);
+	IMPLEMENT_NODE_ATTRIBUTE_SETTER_NODELEGATE(LightShape, uint8);
+}
+
+bool UInterchangeDatasmithAreaLightNode::GetCustomDimensions(FVector2D& AttributeValue) const
+{
+	IMPLEMENT_NODE_ATTRIBUTE_GETTER(Dimensions, FVector2D);
+}
+
+bool UInterchangeDatasmithAreaLightNode::SetCustomDimensions(const FVector2D& AttributeValue, bool bAddApplyDelegate)
+{
+	IMPLEMENT_NODE_ATTRIBUTE_SETTER_NODELEGATE(Dimensions, FVector2D);
+}
+
+bool UInterchangeDatasmithAreaLightNode::GetCustomColor(FLinearColor& AttributeValue) const
+{
+	IMPLEMENT_NODE_ATTRIBUTE_GETTER(Color, FLinearColor);
+}
+
+bool UInterchangeDatasmithAreaLightNode::SetCustomColor(const FLinearColor& AttributeValue, bool bAddApplyDelegate)
+{
+	IMPLEMENT_NODE_ATTRIBUTE_SETTER_NODELEGATE(Color, FLinearColor);
+}
+
+bool UInterchangeDatasmithAreaLightNode::GetCustomIESTexture(FString& AttributeValue) const
+{
+	IMPLEMENT_NODE_ATTRIBUTE_GETTER(IESTexture, FString);
+}
+
+bool UInterchangeDatasmithAreaLightNode::SetCustomIESTexture(const FString& AttributeValue)
+{
+	IMPLEMENT_NODE_ATTRIBUTE_SETTER_NODELEGATE(IESTexture, FString);
+}
+
+bool UInterchangeDatasmithAreaLightNode::GetCustomUseIESBrightness(bool& AttributeValue) const
+{
+	IMPLEMENT_NODE_ATTRIBUTE_GETTER(UseIESBrightness, bool);
+}
+
+bool UInterchangeDatasmithAreaLightNode::SetCustomUseIESBrightness(const bool& AttributeValue, bool bAddApplyDelegate)
+{
+	IMPLEMENT_NODE_ATTRIBUTE_SETTER_NODELEGATE(UseIESBrightness, bool);
+}
+
+bool UInterchangeDatasmithAreaLightNode::GetCustomIESBrightnessScale(float& AttributeValue) const
+{
+	IMPLEMENT_NODE_ATTRIBUTE_GETTER(IESBrightnessScale, float);
+}
+
+bool UInterchangeDatasmithAreaLightNode::SetCustomIESBrightnessScale(const float& AttributeValue, bool bAddApplyDelegate)
+{
+	IMPLEMENT_NODE_ATTRIBUTE_SETTER_NODELEGATE(IESBrightnessScale, float);
+}
+
+bool UInterchangeDatasmithAreaLightNode::GetCustomRotation(FRotator& AttributeValue) const
+{
+	IMPLEMENT_NODE_ATTRIBUTE_GETTER(Rotation, FRotator);
+}
+
+bool UInterchangeDatasmithAreaLightNode::SetCustomRotation(const FRotator& AttributeValue, bool bAddApplyDelegate)
+{
+	IMPLEMENT_NODE_ATTRIBUTE_SETTER_NODELEGATE(Rotation, FRotator);
+}
+
+bool UInterchangeDatasmithAreaLightNode::GetCustomSourceRadius(float& AttributeValue) const
+{
+	IMPLEMENT_NODE_ATTRIBUTE_GETTER(SourceRadius, float);
+}
+
+bool UInterchangeDatasmithAreaLightNode::SetCustomSourceRadius(const float& AttributeValue, bool bAddApplyDelegate)
+{
+	IMPLEMENT_NODE_ATTRIBUTE_SETTER_NODELEGATE(SourceRadius, float);
+}
+
+bool UInterchangeDatasmithAreaLightNode::GetCustomSourceLength(float& AttributeValue) const
+{
+	IMPLEMENT_NODE_ATTRIBUTE_GETTER(SourceLength, float);
+}
+
+bool UInterchangeDatasmithAreaLightNode::SetCustomSourceLength(const float& AttributeValue, bool bAddApplyDelegate)
+{
+	IMPLEMENT_NODE_ATTRIBUTE_SETTER_NODELEGATE(SourceLength, float);
+}
+
+bool UInterchangeDatasmithAreaLightNode::GetCustomSpotlightInnerAngle(float& AttributeValue) const
+{
+	IMPLEMENT_NODE_ATTRIBUTE_GETTER(SpotlightInnerAngle, float);
+}
+
+bool UInterchangeDatasmithAreaLightNode::SetCustomSpotlightInnerAngle(const float& AttributeValue, bool bAddApplyDelegate)
+{
+	IMPLEMENT_NODE_ATTRIBUTE_SETTER_NODELEGATE(SpotlightInnerAngle, float);
+}
+
+bool UInterchangeDatasmithAreaLightNode::GetCustomSpotlightOuterAngle(float& AttributeValue) const
+{
+	IMPLEMENT_NODE_ATTRIBUTE_GETTER(SpotlightOuterAngle, float);
+}
+
+bool UInterchangeDatasmithAreaLightNode::SetCustomSpotlightOuterAngle(const float& AttributeValue, bool bAddApplyDelegate)
+{
+	IMPLEMENT_NODE_ATTRIBUTE_SETTER_NODELEGATE(SpotlightOuterAngle, float);
+}
\ No newline at end of file
diff --git a/Engine/Plugins/Experimental/Enterprise/DatasmithInterchange/Source/DatasmithInterchange/Private/InterchangeDatasmithLevelPipeline.cpp b/Engine/Plugins/Experimental/Enterprise/DatasmithInterchange/Source/DatasmithInterchange/Private/InterchangeDatasmithLevelPipeline.cpp
new file mode 100644
--- /dev/null
+++ b/Engine/Plugins/Experimental/Enterprise/DatasmithInterchange/Source/DatasmithInterchange/Private/InterchangeDatasmithLevelPipeline.cpp
@@ -0,0 +1,207 @@
+// Copyright Epic Games, Inc. All Rights Reserved.
+
+#include "InterchangeDatasmithLevelPipeline.h"
+
+#include "InterchangeDatasmithAreaLightFactoryNode.h"
+#include "InterchangeDatasmithAreaLightNode.h"
+#include "InterchangeDatasmithUtils.h"
+
+
+#include "InterchangeAnimationTrackSetFactoryNode.h"
+#include "InterchangeAnimationTrackSetNode.h"
+#include "InterchangeMaterialFactoryNode.h"
+#include "InterchangeMeshActorFactoryNode.h"
+#include "InterchangeSceneNode.h"
+#include "InterchangeStaticMeshFactoryNode.h"
+
+#include "Components/StaticMeshComponent.h"
+#include "Engine/StaticMesh.h"
+#include "Engine/StaticMeshActor.h"
+#include "Materials/MaterialInterface.h"
+
+void UInterchangeDatasmithLevelPipeline::ExecutePreImportPipeline(UInterchangeBaseNodeContainer* NodeContainer, const TArray<UInterchangeSourceData*>& InSourceDatas)
+{
+	using namespace UE::DatasmithInterchange;
+
+	Super::ExecutePreImportPipeline(NodeContainer, InSourceDatas);
+
+	// Add material factory dependencies for mesh actors where all overrides are filled with the same material
+	for (UInterchangeMeshActorFactoryNode* MeshActorFactoryNode : NodeUtils::GetNodes<UInterchangeMeshActorFactoryNode>(NodeContainer))
+	{
+		// If applicable, transfer material override attribute from translated node to factory node
+		TArray<FString> TargetNodes;
+		MeshActorFactoryNode->GetTargetNodeUids(TargetNodes);
+		if (TargetNodes.Num() == 0)
+		{
+			continue;
+		}
+
+		const UInterchangeSceneNode* SceneNode = Cast<UInterchangeSceneNode>(NodeContainer->GetNode(TargetNodes[0]));
+		if (!SceneNode)
+		{
+			continue;
+		}
+
+		FString MaterialUid;
+		if (!SceneNode->GetStringAttribute(MeshUtils::MeshMaterialAttrName, MaterialUid))
+		{
+			continue;
+		}
+
+		const FString MaterialFactoryUid = UInterchangeMaterialFactoryNode::GetMaterialFactoryNodeUidFromMaterialNodeUid(MaterialUid);
+		MeshActorFactoryNode->AddFactoryDependencyUid(MaterialFactoryUid);
+		MeshActorFactoryNode->AddStringAttribute(MeshUtils::MeshMaterialAttrName, MaterialFactoryUid);
+	}
+}
+
+void UInterchangeDatasmithLevelPipeline::ExecutePostImportPipeline(const UInterchangeBaseNodeContainer* NodeContainer, const FString& FactoryNodeKey, UObject* CreatedAsset, bool bIsAReimport)
+{
+	using namespace UE::DatasmithInterchange;
+
+	if (!NodeContainer || !CreatedAsset)
+	{
+		return;
+	}
+
+	Super::ExecutePostImportPipeline(NodeContainer, FactoryNodeKey, CreatedAsset, bIsAReimport);
+
+	// If applicable, update material overrides for newly created mesh actor
+	AStaticMeshActor* MeshActor = Cast<AStaticMeshActor>(CreatedAsset);
+	UStaticMeshComponent* MeshComponent = MeshActor ? MeshActor->GetStaticMeshComponent() : nullptr;
+	UStaticMesh* StaticMesh = MeshComponent ? MeshComponent->GetStaticMesh() : nullptr;
+	if (MeshActor && MeshComponent && StaticMesh)
+	{
+		const UInterchangeMeshActorFactoryNode* FactoryNode = Cast<UInterchangeMeshActorFactoryNode>(NodeContainer->GetFactoryNode(FactoryNodeKey));
+		if (!FactoryNode)
+		{
+			return;
+		}
+
+		FString MaterialFactoryUid;
+		if (!FactoryNode->GetStringAttribute(MeshUtils::MeshMaterialAttrName, MaterialFactoryUid))
+		{
+			return;
+		}
+
+		const UInterchangeBaseMaterialFactoryNode* MaterialFactoryNode = Cast< UInterchangeBaseMaterialFactoryNode>(NodeContainer->GetFactoryNode(MaterialFactoryUid));
+		if (!MaterialFactoryNode)
+		{
+			return;
+		}
+		FSoftObjectPath MaterialFactoryNodeReferenceObject;
+		MaterialFactoryNode->GetCustomReferenceObject(MaterialFactoryNodeReferenceObject);
+		if (UMaterialInterface* MaterialInterface = Cast<UMaterialInterface>(MaterialFactoryNodeReferenceObject.ResolveObject()))
+		{
+			const TArray<FStaticMaterial>& StaticMaterials = StaticMesh->GetStaticMaterials();
+			for (int32 Index = 0; Index < StaticMaterials.Num(); ++Index)
+			{
+				MeshComponent->SetMaterial(Index, MaterialInterface);
+			}
+		}
+	}
+
+	// If a ULevelSequence is referencing a ADatasmithAreaLightActor, update its mobility property
+	if (const UInterchangeAnimationTrackSetFactoryNode* FactoryNode = Cast<UInterchangeAnimationTrackSetFactoryNode>(NodeContainer->GetFactoryNode(FactoryNodeKey)))
+	{
+		FSoftObjectPath FactoryNodeReferenceObject;
+		FactoryNode->GetCustomReferenceObject(FactoryNodeReferenceObject);
+		if (FactoryNodeReferenceObject.TryLoad() == CreatedAsset)
+		{
+			TArray<FString> AnimationTrackUids;
+			FactoryNode->GetCustomAnimationTrackUids(AnimationTrackUids);
+
+			for (const FString& AnimationTrackUid : AnimationTrackUids)
+			{
+				if (const UInterchangeBaseNode* TranslatedNode = NodeContainer->GetNode(AnimationTrackUid))
+				{
+					if (const UInterchangeTransformAnimationTrackNode* TransformTrackNode = Cast<UInterchangeTransformAnimationTrackNode>(TranslatedNode))
+					{
+						FString ActorNodeUid;
+						if (TransformTrackNode->GetCustomActorDependencyUid(ActorNodeUid))
+						{
+							const FString ActorFactoryNodeUid = TEXT("Factory_") + ActorNodeUid;
+							const UInterchangeFactoryBaseNode* ActorFactoryNode = Cast<UInterchangeFactoryBaseNode>(NodeContainer->GetNode(ActorFactoryNodeUid));
+
+							if (ActorFactoryNode)
+							{
+								FSoftObjectPath ActorFactoryNodeReferenceObject;
+								ActorFactoryNode->GetCustomReferenceObject(FactoryNodeReferenceObject);
+								if (ADatasmithAreaLightActor* DatasmithAreaLightActor = Cast<ADatasmithAreaLightActor>(ActorFactoryNodeReferenceObject.TryLoad()))
+								{
+									DatasmithAreaLightActor->Mobility = EComponentMobility::Movable;
+								}
+							}
+						}
+					}
+				}
+			}
+		}
+	}
+}
+
+UInterchangeActorFactoryNode* UInterchangeDatasmithLevelPipeline::CreateActorFactoryNode(const UInterchangeSceneNode* SceneNode, const UInterchangeBaseNode* TranslatedAssetNode) const
+{
+	if (!ensure(BaseNodeContainer))
+	{
+		return nullptr;
+	}
+
+	if (TranslatedAssetNode && TranslatedAssetNode->IsA<UInterchangeDatasmithAreaLightNode>())
+	{
+		return NewObject<UInterchangeDatasmithAreaLightFactoryNode>(BaseNodeContainer, NAME_None);
+	}
+	else
+	{
+		return Super::CreateActorFactoryNode(SceneNode, TranslatedAssetNode);
+	}
+}
+
+void UInterchangeDatasmithLevelPipeline::SetUpFactoryNode(UInterchangeActorFactoryNode* ActorFactoryNode, const UInterchangeSceneNode* SceneNode, const UInterchangeBaseNode* TranslatedAssetNode) const
+{
+	if (const UInterchangeDatasmithAreaLightNode* AreaLightNode = Cast<UInterchangeDatasmithAreaLightNode>(TranslatedAssetNode))
+	{
+		UInterchangeDatasmithAreaLightFactoryNode* AreaLightFactory = Cast<UInterchangeDatasmithAreaLightFactoryNode>(ActorFactoryNode);
+		ensure(AreaLightFactory);
+		SetupAreaLight(AreaLightFactory, AreaLightNode);
+	}
+	else
+	{
+		Super::SetUpFactoryNode(ActorFactoryNode, SceneNode, TranslatedAssetNode);
+	}
+}
+
+#define APPLY_FACTORY_ATTRIBUTE(AttributeName, AttributeType) \
+	AttributeType AttributeName;\
+	if (TranslatedNode->GetCustom##AttributeName(AttributeName))\
+	{\
+		FactoryNode->SetCustom##AttributeName(AttributeName);\
+	}\
+
+void UInterchangeDatasmithLevelPipeline::SetupAreaLight(UInterchangeDatasmithAreaLightFactoryNode* FactoryNode, const UInterchangeDatasmithAreaLightNode* TranslatedNode) const
+{
+	APPLY_FACTORY_ATTRIBUTE(LightType, EDatasmithAreaLightActorType);
+	APPLY_FACTORY_ATTRIBUTE(LightShape, EDatasmithAreaLightActorShape);
+	APPLY_FACTORY_ATTRIBUTE(Dimensions, FVector2D);
+	
+	APPLY_FACTORY_ATTRIBUTE(Intensity, float);
+	EInterchangeLightUnits IntensityUnits;
+	if(TranslatedNode->GetCustomIntensityUnits(IntensityUnits))
+	{
+		FactoryNode->SetCustomIntensityUnits(ELightUnits(IntensityUnits));
+	}
+	APPLY_FACTORY_ATTRIBUTE(Color, FLinearColor);
+	APPLY_FACTORY_ATTRIBUTE(Temperature, float);
+
+	APPLY_FACTORY_ATTRIBUTE(IESTexture, FString);
+	APPLY_FACTORY_ATTRIBUTE(UseIESBrightness, bool);
+	APPLY_FACTORY_ATTRIBUTE(IESBrightnessScale, float);
+	APPLY_FACTORY_ATTRIBUTE(Rotation, FRotator);
+
+	APPLY_FACTORY_ATTRIBUTE(SourceRadius, float);
+	APPLY_FACTORY_ATTRIBUTE(SourceLength, float);
+	APPLY_FACTORY_ATTRIBUTE(AttenuationRadius, float);
+	APPLY_FACTORY_ATTRIBUTE(SpotlightInnerAngle, float);
+	APPLY_FACTORY_ATTRIBUTE(SpotlightOuterAngle, float);
+}
+
+#undef APPLY_FACTORY_ATTRIBUTE
\ No newline at end of file
diff --git a/Engine/Plugins/Experimental/Enterprise/DatasmithInterchange/Source/DatasmithInterchange/Private/InterchangeDatasmithLevelPipeline.h b/Engine/Plugins/Experimental/Enterprise/DatasmithInterchange/Source/DatasmithInterchange/Private/InterchangeDatasmithLevelPipeline.h
new file mode 100644
--- /dev/null
+++ b/Engine/Plugins/Experimental/Enterprise/DatasmithInterchange/Source/DatasmithInterchange/Private/InterchangeDatasmithLevelPipeline.h
@@ -0,0 +1,33 @@
+// Copyright Epic Games, Inc. All Rights Reserved.
+
+#pragma once
+
+#include "CoreMinimal.h"
+#include "InterchangeGenericScenesPipeline.h"
+#include "UObject/Object.h"
+#include "UObject/ObjectMacros.h"
+
+#include "InterchangeDatasmithLevelPipeline.generated.h"
+
+class UInterchangeBaseNodeContainer;
+class UInterchangeSceneNode;
+class UInterchangeDatasmithAreaLightFactoryNode;
+class UInterchangeDatasmithAreaLightNode;
+
+UCLASS(BlueprintType, Experimental)
+class UInterchangeDatasmithLevelPipeline : public UInterchangeGenericLevelPipeline
+{
+	GENERATED_BODY()
+
+protected:
+	virtual void ExecutePreImportPipeline(UInterchangeBaseNodeContainer* BaseNodeContainer, const TArray<UInterchangeSourceData*>& SourceDatas) override;
+	virtual void ExecutePostImportPipeline(const UInterchangeBaseNodeContainer* BaseNodeContainer, const FString& NodeKey, UObject* CreatedAsset, bool bIsAReimport) override;
+
+	virtual UInterchangeActorFactoryNode* CreateActorFactoryNode(const UInterchangeSceneNode* SceneNode, const UInterchangeBaseNode* TranslatedAssetNode) const override;
+
+	virtual void SetUpFactoryNode(UInterchangeActorFactoryNode* ActorFactoryNode, const UInterchangeSceneNode* SceneNode, const UInterchangeBaseNode* TranslatedAssetNode) const override;
+
+private:
+
+	void SetupAreaLight(UInterchangeDatasmithAreaLightFactoryNode* AreaLightFactoryNode, const UInterchangeDatasmithAreaLightNode* AreaLightNode) const;
+};
\ No newline at end of file
diff --git a/Engine/Plugins/Experimental/Enterprise/DatasmithInterchange/Source/DatasmithInterchange/Private/InterchangeDatasmithLog.h b/Engine/Plugins/Experimental/Enterprise/DatasmithInterchange/Source/DatasmithInterchange/Private/InterchangeDatasmithLog.h
new file mode 100644
--- /dev/null
+++ b/Engine/Plugins/Experimental/Enterprise/DatasmithInterchange/Source/DatasmithInterchange/Private/InterchangeDatasmithLog.h
@@ -0,0 +1,7 @@
+// Copyright Epic Games, Inc. All Rights Reserved.
+
+#pragma once
+
+#include "CoreMinimal.h"
+
+DECLARE_LOG_CATEGORY_EXTERN(LogInterchangeDatasmith, Log, All);
diff --git a/Engine/Plugins/Experimental/Enterprise/DatasmithInterchange/Source/DatasmithInterchange/Private/InterchangeDatasmithMaterialNode.cpp b/Engine/Plugins/Experimental/Enterprise/DatasmithInterchange/Source/DatasmithInterchange/Private/InterchangeDatasmithMaterialNode.cpp
new file mode 100644
--- /dev/null
+++ b/Engine/Plugins/Experimental/Enterprise/DatasmithInterchange/Source/DatasmithInterchange/Private/InterchangeDatasmithMaterialNode.cpp
@@ -0,0 +1,948 @@
+// Copyright Epic Games, Inc. All Rights Reserved. 
+
+#include "InterchangeDatasmithMaterialNode.h"
+#include "InterchangeDatasmithPipeline.h"
+#include "InterchangeDatasmithUtils.h"
+
+#include "DatasmithMaterialElements.h"
+#include "DatasmithUtils.h"
+#include "IDatasmithSceneElements.h"
+
+#include "InterchangeReferenceMaterials/DatasmithReferenceMaterialManager.h"
+#include "InterchangeReferenceMaterials/DatasmithReferenceMaterialSelector.h"
+
+#include "InterchangeManager.h"
+#include "InterchangeMaterialDefinitions.h"
+#include "InterchangeMaterialFactoryNode.h"
+#include "InterchangeShaderGraphNode.h"
+#include "InterchangeTexture2DNode.h"
+#include "InterchangeTexture2DFactoryNode.h"
+
+#include "Async/TaskGraphInterfaces.h"
+#include "Materials/MaterialExpression.h"
+#include "Materials/MaterialExpressionCustom.h"
+#include "Materials/MaterialExpressionCustomOutput.h"
+#include "Materials/MaterialExpressionMaterialFunctionCall.h"
+#include "Materials/MaterialFunction.h"
+#include "Materials/MaterialInstanceConstant.h"
+#include "Materials/MaterialInstanceDynamic.h"
+#include "Misc/PackageName.h"
+
+const FName UInterchangeDatasmithMaterialNode::MaterialTypeAttrName(TEXT("Datasmith:ReferenceMaterial:MaterialType"));
+const FName UInterchangeDatasmithMaterialNode::MaterialQualityAttrName(TEXT("Datasmith:ReferenceMaterial:MaterialQuality"));
+const FName UInterchangeDatasmithMaterialNode::MaterialParentAttrName(TEXT("Datasmith:ReferenceMaterial:MaterialParent"));
+
+const FName UInterchangeDatasmithPbrMaterialNode::ShadingModelAttrName(TEXT("Datasmith:UEPbrMaterial:ShadingModel"));
+const FName UInterchangeDatasmithPbrMaterialNode::BlendModeAttrName(TEXT("Datasmith:UEPbrMaterial:BlendMode"));
+const FName UInterchangeDatasmithPbrMaterialNode::OpacityMaskClipValueAttrName(TEXT("Datasmith:UEPbrMaterial:OpacityMaskClipValue"));
+const FName UInterchangeDatasmithPbrMaterialNode::TranslucencyLightingModeAttrName(TEXT("Datasmith:UEPbrMaterial:TranslucencyLightingMode"));
+const FString UInterchangeDatasmithPbrMaterialNode::MaterialFunctionsDependenciesKey(TEXT("Datasmith:UEPbrMaterial:MaterialFunctionsDependencies"));
+
+namespace UE::DatasmithInterchange::MaterialUtils
+{
+	const FName MaterialFunctionPathAttrName(TEXT("Datasmith:Material:FunctionCall:MaterialFunctionPath"));
+	const FName DefaultOutputIndexAttrName(TEXT("Datasmith:MaterialExpression:DefaultOutputIndex"));
+
+#if WITH_EDITOR
+	class FPbrMaterialHelper
+	{
+	public:
+		UInterchangeBaseNodeContainer& NodeContainer;
+		FString MaterialID;
+		IDatasmithUEPbrMaterialElement& MaterialElement;
+		TMap<const IDatasmithMaterialExpression*, int32> ExpresionElements;
+		TArray<UInterchangeShaderNode*> ProcessedExpressions;
+		TSet<FString> ExpressionNames;
+		TSet<FString> MaterialFunctionUids;
+
+		struct FConnectionData
+		{
+			UInterchangeShaderNode* InputNode = nullptr;
+			FString InputName;
+
+			void ConnectDefaultOuputToInput(const FString& OutputNodeUid) const
+			{
+				UInterchangeShaderPortsAPI::ConnectDefaultOuputToInput(InputNode, InputName, OutputNodeUid);
+			}
+
+			void ConnectOuputToInput(const FString& OutputNodeUid, const FName& OutputName) const
+			{
+				UInterchangeShaderPortsAPI::ConnectOuputToInput(InputNode, InputName, OutputNodeUid, OutputName.ToString());
+			}
+		};
+
+		FPbrMaterialHelper(UInterchangeShaderNode& MaterialNode, IDatasmithUEPbrMaterialElement& InMaterialElement);
+
+		void ConnectExpression(const IDatasmithMaterialExpressionTexture& DatasmithExpression, const FConnectionData& ConnectionData, int32 OutputIndex);
+
+		void ConnectExpression(const IDatasmithMaterialExpressionTextureCoordinate& DatasmithExpression, const FConnectionData& ConnectionData);
+
+		void ConnectExpression(const IDatasmithMaterialExpressionFlattenNormal& DatasmithExpression, const FConnectionData& ConnectionData);
+
+		void ConnectExpression(const IDatasmithMaterialExpressionGeneric& DatasmithExpression, const FConnectionData& ConnectionData, int32 OutputIndex);
+
+		void ConnectExpression(const IDatasmithMaterialExpressionFunctionCall& DatasmithExpression, const FConnectionData& ConnectionData, int32 OutputIndex);
+
+		void ConnectExpression(const IDatasmithMaterialExpressionCustom& DatasmithExpression, const FConnectionData& ConnectionData, int32 OutputIndex);
+
+		void ConnectExpression(const IDatasmithMaterialExpression* MaterialExpression, const FConnectionData& ConnectionData, int32 OutputIndex);
+
+		static UClass* FindClass(const TCHAR* ClassName);
+
+		static UClass* GetGenericExpressionClass(const TCHAR* GenericExpressionname);
+
+		static UMaterialExpression* GetDefaultMaterialExpression(const TCHAR* GenericExpressionname);
+
+		UInterchangeShaderNode* BuildGenericExpression(const IDatasmithMaterialExpressionGeneric& DatasmithExpression, UInterchangeShaderNode& ParentNode, TArray<FName>* OutputNames = nullptr);
+
+		template<class T = UInterchangeShaderNode>
+		T* CreateExpressionNode(const IDatasmithMaterialExpression& DatasmithExpression, const FString& ParentNodeUid)
+		{
+			int32* ExpressionIndexPtr = ExpresionElements.Find(&DatasmithExpression);
+			if (!ExpressionIndexPtr)
+			{
+				// TODO: Log error
+				return nullptr;
+			}
+
+			const FString NodeName = FDatasmithUtils::SanitizeObjectName(DatasmithExpression.GetName());
+			const FString NodeUid = NodeUtils::MaterialExpressionPrefix + MaterialID + TEXT("_") + FString::FromInt(*ExpressionIndexPtr) + TEXT("_") + NodeName;
+
+			UInterchangeShaderNode* ExpressionNode = NewObject< T >(&NodeContainer);
+			ExpressionNode->InitializeNode(NodeUid, NodeName, EInterchangeNodeContainerType::TranslatedAsset);
+			NodeContainer.AddNode(ExpressionNode);
+
+			NodeContainer.SetNodeParentUid(NodeUid, ParentNodeUid);
+
+			ExpressionNode->AddInt32Attribute(DefaultOutputIndexAttrName, DatasmithExpression.GetDefaultOutputIndex());
+			ProcessedExpressions[*ExpressionIndexPtr] = ExpressionNode;
+
+			return Cast<T>(ExpressionNode);
+		}
+
+		EBlendMode GetUEPbrImportBlendMode();
+	};
+#endif
+
+	bool BuildMaterialNode(IDatasmithMaterialInstanceElement& MaterialElement, UInterchangeDatasmithMaterialNode& MaterialNode)
+	{
+		using namespace UE::Interchange::Materials;
+
+		MaterialNode.SetMaterialType(MaterialElement.GetMaterialType());
+		MaterialNode.SetMaterialQuality(MaterialElement.GetQuality());
+
+
+		UInterchangeBaseNodeContainer* NodeContainer = Cast<UInterchangeBaseNodeContainer>(MaterialNode.GetOuter());
+		check(NodeContainer);
+
+		// Find matching reference material parameters
+		for (int Index = 0; Index < MaterialElement.GetPropertiesCount(); ++Index)
+		{
+			const TSharedPtr< IDatasmithKeyValueProperty > Property = MaterialElement.GetProperty(Index);
+			const FName InputValueKey = UInterchangeShaderPortsAPI::MakeInputValueKey(FDatasmithUtils::SanitizeObjectName(Property->GetName()));
+
+			switch (Property->GetPropertyType())
+			{
+				case EDatasmithKeyValuePropertyType::Color:
+				{
+					FLinearColor LinearColor;
+					LinearColor.InitFromString(Property->GetValue());
+
+					MaterialNode.AddLinearColorAttribute(InputValueKey, LinearColor);
+					break;
+				}
+
+				case EDatasmithKeyValuePropertyType::Float:
+				{
+					MaterialNode.AddFloatAttribute(InputValueKey, FCString::Atof(Property->GetValue()));
+					break;
+				}
+
+				case EDatasmithKeyValuePropertyType::Bool:
+				{
+					MaterialNode.AddBooleanAttribute(InputValueKey, FString(Property->GetValue()).ToBool());
+					break;
+				}
+
+				case EDatasmithKeyValuePropertyType::Integer:
+				{
+					MaterialNode.AddInt32Attribute(InputValueKey, FCString::Atoi(Property->GetValue()));
+					break;
+				}
+
+				case EDatasmithKeyValuePropertyType::Texture:
+				{
+					MaterialNode.AddStringAttribute(InputValueKey, Property->GetValue());
+					break;
+				}
+
+				case EDatasmithKeyValuePropertyType::Vector:
+				case EDatasmithKeyValuePropertyType::String:
+				default:
+				{
+					break;
+				}
+			}
+		}
+
+		return true;
+	}
+
+	bool BuildMaterialNode(IDatasmithMaterialElement& MaterialElement, UInterchangeDatasmithMaterialNode& MaterialNode)
+	{
+		using namespace UE::Interchange::Materials;
+
+		return false;
+	}
+
+
+#if WITH_EDITOR
+	FPbrMaterialHelper::FPbrMaterialHelper(UInterchangeShaderNode& MaterialNode, IDatasmithUEPbrMaterialElement& InMaterialElement)
+		: NodeContainer(*Cast<UInterchangeBaseNodeContainer>(MaterialNode.GetOuter()))
+		, MaterialID(MaterialNode.GetUniqueID())
+		, MaterialElement(InMaterialElement)
+	{
+		const int32 ExpressionsCount = MaterialElement.GetExpressionsCount();
+		for (int32 ExpressionIndex = 0; ExpressionIndex < ExpressionsCount; ++ExpressionIndex)
+		{
+			if (const IDatasmithMaterialExpression* MaterialExpression = MaterialElement.GetExpression(ExpressionIndex))
+			{
+				ExpresionElements.Add(MaterialExpression, ExpressionIndex);
+			}
+		}
+
+		ProcessedExpressions.AddZeroed(ExpressionsCount);
+	}
+
+	void FPbrMaterialHelper::ConnectExpression(const IDatasmithMaterialExpressionTexture& DatasmithExpression, const FConnectionData& ConnectionData, int32 OutputIndex)
+	{
+		using namespace UE::Interchange::Materials::Standard::Nodes;
+
+		if (OutputIndex < 0 || OutputIndex > 5)
+		{
+			// TODO: Log error
+			return;
+		}
+
+		UInterchangeShaderNode* ExpressionNode = CreateExpressionNode(DatasmithExpression, ConnectionData.InputNode->GetUniqueID());
+		if (!ensure(ExpressionNode))
+		{
+			// TODO: Log error
+			return;
+		}
+
+		if (FPackageName::IsValidObjectPath(DatasmithExpression.GetTexturePathName()))
+		{
+			// TODO: Handle case of existing asset
+		}
+		else
+		{
+			ExpressionNode->SetCustomShaderType(TextureSample::Name.ToString());
+
+			const FString TextureUid = NodeUtils::TexturePrefix + FDatasmithUtils::SanitizeObjectName(DatasmithExpression.GetTexturePathName());
+			if (NodeContainer.GetNode(TextureUid))
+			{
+				const FName ValueKey = UInterchangeShaderPortsAPI::MakeInputValueKey(TextureSample::Inputs::Texture.ToString());
+				ExpressionNode->AddStringAttribute(ValueKey, TextureUid);
+			}
+			else
+			{
+				ensure(false);
+				// TODO: Log warning
+			}
+		}
+
+		static FName OutputNames[6] = { TextureSample::Outputs::RGB,
+										TextureSample::Outputs::R,
+										TextureSample::Outputs::G,
+										TextureSample::Outputs::B,
+										TextureSample::Outputs::A,
+										TextureSample::Outputs::RGBA };
+
+		ConnectionData.ConnectOuputToInput(ExpressionNode->GetUniqueID(), OutputNames[OutputIndex]);
+
+		FConnectionData CoordConnectionData{ExpressionNode, TextureSample::Inputs::Coordinates.ToString() };
+
+		const IDatasmithExpressionInput& ExpressionInput = DatasmithExpression.GetInputCoordinate();
+		ConnectExpression(ExpressionInput.GetExpression(), CoordConnectionData, ExpressionInput.GetOutputIndex());
+	}
+
+	void FPbrMaterialHelper::ConnectExpression(const IDatasmithMaterialExpressionTextureCoordinate& DatasmithExpression, const FConnectionData& ConnectionData)
+	{
+		using namespace UE::Interchange::Materials::Standard::Nodes;
+
+		UInterchangeShaderNode* ExpressionNode = CreateExpressionNode(DatasmithExpression, ConnectionData.InputNode->GetUniqueID());
+		if (!ensure(ExpressionNode))
+		{
+			// TODO: Log error
+			return;
+		}
+
+		ExpressionNode->SetCustomShaderType(TextureCoordinate::Name.ToString());
+
+		{
+			const FName ValueKey = UInterchangeShaderPortsAPI::MakeInputValueKey(TextureCoordinate::Inputs::Index.ToString());
+			ExpressionNode->AddInt32Attribute(ValueKey, DatasmithExpression.GetCoordinateIndex());
+		}
+
+		{
+			const FName ValueKey = UInterchangeShaderPortsAPI::MakeInputValueKey(TextureCoordinate::Inputs::UTiling.ToString());
+			ExpressionNode->AddFloatAttribute(ValueKey, DatasmithExpression.GetUTiling());
+		}
+
+		{
+			const FName ValueKey = UInterchangeShaderPortsAPI::MakeInputValueKey(TextureCoordinate::Inputs::VTiling.ToString());
+			ExpressionNode->AddFloatAttribute(ValueKey, DatasmithExpression.GetVTiling());
+		}
+
+		ConnectionData.ConnectDefaultOuputToInput(ExpressionNode->GetUniqueID());
+	}
+
+	void FPbrMaterialHelper::ConnectExpression(const IDatasmithMaterialExpressionFlattenNormal& DatasmithExpression, const FConnectionData& ConnectionData)
+	{
+		using namespace UE::Interchange::Materials::Standard::Nodes;
+
+		UInterchangeShaderNode* ExpressionNode = CreateExpressionNode(DatasmithExpression, ConnectionData.InputNode->GetUniqueID());
+		if (!ensure(ExpressionNode))
+		{
+			// TODO: Log error
+			return;
+		}
+
+		ExpressionNode->SetCustomShaderType(FlattenNormal::Name.ToString());
+
+		{
+			FConnectionData InputConnectionData{ ExpressionNode, FlattenNormal::Inputs::Normal.ToString() };
+
+			const IDatasmithExpressionInput& ExpressionInput = DatasmithExpression.GetNormal();
+			ConnectExpression(ExpressionInput.GetExpression(), InputConnectionData, ExpressionInput.GetOutputIndex());
+		}
+
+		{
+			FConnectionData InputConnectionData{ ExpressionNode, FlattenNormal::Inputs::Flatness.ToString() };
+
+			const IDatasmithExpressionInput& ExpressionInput = DatasmithExpression.GetFlatness();
+			ConnectExpression(ExpressionInput.GetExpression(), InputConnectionData, ExpressionInput.GetOutputIndex());
+		}
+
+		ConnectionData.ConnectDefaultOuputToInput(ExpressionNode->GetUniqueID());
+	}
+
+	UClass* FPbrMaterialHelper::FindClass(const TCHAR* ClassName)
+	{
+		check(ClassName);
+
+		if (UClass* Result = FindFirstObject<UClass>(ClassName, EFindFirstObjectOptions::None, ELogVerbosity::Warning, TEXT("Datasmith FindClass")))
+		{
+			return Result;
+		}
+
+		if (UObjectRedirector* RenamedClassRedirector = FindFirstObject<UObjectRedirector>(ClassName, EFindFirstObjectOptions::None, ELogVerbosity::Warning, TEXT("Datasmith FindClass")))
+		{
+			return CastChecked<UClass>(RenamedClassRedirector->DestinationObject);
+		}
+
+		return nullptr;
+	}
+
+	UClass* FPbrMaterialHelper::GetGenericExpressionClass(const TCHAR* GenericExpressionname)
+	{
+		const TCHAR* MaterialExpressionCharPtr = TEXT("MaterialExpression");
+		const int32 MaterialExpressionLength = FCString::Strlen(MaterialExpressionCharPtr);
+		const int32 ExpressionNameLength = FCString::Strlen(GenericExpressionname);
+
+		FString ClassName;
+		ClassName.Reserve(MaterialExpressionLength + ExpressionNameLength);
+		ClassName.AppendChars(TEXT("MaterialExpression"), MaterialExpressionLength);
+		ClassName.AppendChars(GenericExpressionname, ExpressionNameLength);
+		return FindClass(*ClassName);
+	}
+
+	UMaterialExpression* FPbrMaterialHelper::GetDefaultMaterialExpression(const TCHAR* GenericExpressionname)
+	{
+		UClass* ExpressionClass = GetGenericExpressionClass(GenericExpressionname);
+
+		if (!ensure(ExpressionClass))
+		{
+			return nullptr;
+		}
+
+		return ExpressionClass->GetDefaultObject<UMaterialExpression>();
+	}
+
+	UInterchangeShaderNode* FPbrMaterialHelper::BuildGenericExpression(const IDatasmithMaterialExpressionGeneric& DatasmithExpression, UInterchangeShaderNode& ParentNode, TArray<FName>* OutputNames)
+	{
+		UMaterialExpression* MaterialExpression = GetDefaultMaterialExpression(DatasmithExpression.GetExpressionName());
+
+		if (!ensure(MaterialExpression))
+		{
+				// TODO: Log error
+			return nullptr;
+		}
+
+		UInterchangeShaderNode* ExpressionNode = CreateExpressionNode(DatasmithExpression, ParentNode.GetUniqueID());
+		if (!ensure(ExpressionNode))
+		{
+				// TODO: Log error
+			return nullptr;
+		}
+
+		ExpressionNode->SetCustomShaderType(DatasmithExpression.GetExpressionName());
+
+		// TODO: What about if MaterialExpression->HasAParameterName()?
+
+		for (int32 InputIndex = 0; InputIndex < DatasmithExpression.GetInputCount(); ++InputIndex)
+		{
+			if (const IDatasmithExpressionInput* ExpressionInput = DatasmithExpression.GetInput(InputIndex))
+			{
+				const FString InputName = MaterialExpression->GetInputName(InputIndex).ToString();
+				FConnectionData InputConnectionData{ ExpressionNode, InputName };
+
+				ConnectExpression(ExpressionInput->GetExpression(), InputConnectionData, ExpressionInput->GetOutputIndex());
+			}
+		}
+
+		if (OutputNames != nullptr)
+		{
+			TArray<FExpressionOutput>& Outputs = MaterialExpression->GetOutputs();
+
+			OutputNames->Reserve(Outputs.Num());
+
+			for (FExpressionOutput& ExpressionOutput : Outputs)
+			{
+				OutputNames->Add(ExpressionOutput.OutputName);
+			}
+		}
+
+		return ExpressionNode;
+	}
+
+	void FPbrMaterialHelper::ConnectExpression(const IDatasmithMaterialExpressionGeneric& DatasmithExpression, const FConnectionData& ConnectionData, int32 OutputIndex)
+	{
+		TArray<FName> OutputNames;
+		UInterchangeShaderNode* ExpressionNode = BuildGenericExpression(DatasmithExpression, *ConnectionData.InputNode, &OutputNames);
+		if (!ensure(ExpressionNode))
+		{
+			// TODO: Log error
+			return;
+		}
+
+		if (!ensure(OutputNames.IsValidIndex(OutputIndex)))
+		{
+			// TODO: Log error
+			return;
+		}
+
+		UMaterialExpression* MaterialExpression = GetDefaultMaterialExpression(DatasmithExpression.GetExpressionName());
+
+		// TODO: Handle case of Mask
+		FExpressionOutput& ExpressionOutput = MaterialExpression->GetOutputs()[OutputIndex];
+		ConnectionData.ConnectOuputToInput(ExpressionNode->GetUniqueID(), ExpressionOutput.OutputName);
+	}
+
+	void FPbrMaterialHelper::ConnectExpression(const IDatasmithMaterialExpressionFunctionCall& DatasmithExpression, const FConnectionData& ConnectionData, int32 OutputIndex)
+	{
+		using namespace UE::Interchange::Materials::Standard::Nodes;
+
+		FString MaterialUid;
+		FString FunctionPath = FPackageName::ExportTextPathToObjectPath(DatasmithExpression.GetFunctionPathName());
+
+		TArray<FString> InputNames;
+		UInterchangeShaderNode* ExpressionNode = nullptr;
+
+		if (FPackageName::DoesPackageExist(FunctionPath))
+		{
+			FSoftObjectPath ObjectPath(TEXT("MaterialFunction'") + FunctionPath + TEXT("'"));
+
+			UMaterialFunctionInterface* MaterialFunction = Cast<UMaterialFunctionInterface>(ObjectPath.TryLoad());
+
+			// TODO: The call to TryLoad may fail if the asset was not loaded ahead of time.
+			// Any asset starting with /Engine may work
+			if (!ensure(MaterialFunction))
+			{
+				//TODO: Log warning
+				return;
+			}
+
+			TArray<FFunctionExpressionInput> FunctionInputs;
+			TArray<FFunctionExpressionOutput> FunctionOutputs;
+
+			MaterialFunction->UpdateFromFunctionResource();
+			MaterialFunction->GetInputsAndOutputs(FunctionInputs, FunctionOutputs);
+
+			const bool bCanProceed = FunctionInputs.Num() < DatasmithExpression.GetInputCount() || FunctionOutputs.IsValidIndex(OutputIndex);
+			if (!ensure(bCanProceed))
+			{
+				// TODO: Log error
+				return;
+			}
+
+			InputNames.Reserve(FunctionInputs.Num());
+			for (FFunctionExpressionInput& FunctionInput : FunctionInputs)
+			{
+				InputNames.Add(FunctionInput.Input.InputName.ToString());
+			}
+
+			ExpressionNode = CreateExpressionNode(DatasmithExpression, ConnectionData.InputNode->GetUniqueID());
+
+			ExpressionNode->AddStringAttribute(MaterialFunctionPathAttrName, FunctionPath);
+		}
+		else
+		{
+			MaterialUid = NodeUtils::MaterialFunctionPrefix + FunctionPath;
+			const UInterchangeDatasmithPbrMaterialNode* PbrMaterialNode = Cast<const UInterchangeDatasmithPbrMaterialNode>(NodeContainer.GetNode(MaterialUid));
+			if (!ensure(PbrMaterialNode))
+			{
+				// TODO: Log error
+				return;
+			}
+
+			UInterchangeShaderPortsAPI::GatherInputs(PbrMaterialNode, InputNames);
+
+			UInterchangeFunctionCallShaderNode* FunctionCallExpressionNode = CreateExpressionNode<UInterchangeFunctionCallShaderNode>(DatasmithExpression, ConnectionData.InputNode->GetUniqueID());
+			FunctionCallExpressionNode->SetCustomMaterialFunction(MaterialUid);
+
+			MaterialFunctionUids.Add(MaterialUid);
+
+			ExpressionNode = FunctionCallExpressionNode;
+		}
+
+		if (!ensure(ExpressionNode))
+		{
+			// TODO: Log error
+			return;
+		}
+
+		// TODO: Handle case where path is an existing asset vs a material function defined in file
+		ExpressionNode->SetCustomShaderType(TEXT("MaterialFunctionCall"));
+
+		ConnectionData.ConnectDefaultOuputToInput(ExpressionNode->GetUniqueID());
+
+		for (int32 InputIndex = 0; InputIndex < DatasmithExpression.GetInputCount(); ++InputIndex)
+		{
+			if (const IDatasmithExpressionInput* ExpressionInput = DatasmithExpression.GetInput(InputIndex))
+			{
+				FConnectionData InputConnectionData{ ExpressionNode, InputNames[InputIndex]};
+
+				ConnectExpression(ExpressionInput->GetExpression(), InputConnectionData, ExpressionInput->GetOutputIndex());
+			}
+		}
+	}
+
+	void FPbrMaterialHelper::ConnectExpression(const IDatasmithMaterialExpressionCustom& DatasmithExpression, const FConnectionData& ConnectionData, int32 OutputIndex)
+	{
+		// TODO: See FDatasmithMaterialExpressions::CreateCustomExpression
+	}
+
+	void FPbrMaterialHelper::ConnectExpression(const IDatasmithMaterialExpression* MaterialExpression, const FConnectionData& ConnectionData, int32 OutputIndex)
+	{
+		if (!MaterialExpression)
+		{
+			return;
+		}
+
+		int32* ExpressionIndexPtr = ExpresionElements.Find(MaterialExpression);
+		if (ExpressionIndexPtr && ProcessedExpressions[*ExpressionIndexPtr])
+		{
+			UInterchangeShaderNode* ExpressionNode = ProcessedExpressions[*ExpressionIndexPtr];
+
+			int32 DefaultOutputIndex = 0;
+			ExpressionNode->GetInt32Attribute(DefaultOutputIndexAttrName, DefaultOutputIndex);
+
+			if (OutputIndex == DefaultOutputIndex)
+			{
+					ConnectionData.ConnectDefaultOuputToInput(ExpressionNode->GetUniqueID());
+			}
+			else
+			{
+				// TODO: Need to find node's output list
+				ConnectionData.ConnectOuputToInput(ExpressionNode->GetUniqueID(), NAME_None);
+			}
+		}
+
+		if (MaterialExpression->IsSubType(EDatasmithMaterialExpressionType::ConstantBool))
+		{
+			const IDatasmithMaterialExpressionBool* BoolExpression = static_cast<const IDatasmithMaterialExpressionBool*>(MaterialExpression);
+			const FName InputName(UInterchangeShaderPortsAPI::MakeInputValueKey(ConnectionData.InputName));
+			ConnectionData.InputNode->AddBooleanAttribute(InputName, BoolExpression->GetBool());
+		}
+		else if (MaterialExpression->IsSubType(EDatasmithMaterialExpressionType::ConstantScalar))
+		{
+			const IDatasmithMaterialExpressionScalar* ScalarExpression = static_cast<const IDatasmithMaterialExpressionScalar*>(MaterialExpression);
+			const FName InputName(UInterchangeShaderPortsAPI::MakeInputValueKey(ConnectionData.InputName));
+			ConnectionData.InputNode->AddFloatAttribute(InputName, ScalarExpression->GetScalar());
+		}
+		else if (MaterialExpression->IsSubType(EDatasmithMaterialExpressionType::ConstantColor))
+		{
+			const IDatasmithMaterialExpressionColor* ColorExpression = static_cast<const IDatasmithMaterialExpressionColor*>(MaterialExpression);
+			const FName InputName(UInterchangeShaderPortsAPI::MakeInputValueKey(ConnectionData.InputName));
+			ConnectionData.InputNode->AddLinearColorAttribute(InputName, ColorExpression->GetColor());
+		}
+		else if (MaterialExpression->IsSubType(EDatasmithMaterialExpressionType::Texture))
+		{
+			ConnectExpression(static_cast<const IDatasmithMaterialExpressionTexture&>(*MaterialExpression), ConnectionData, OutputIndex);
+		}
+		else if (MaterialExpression->IsSubType(EDatasmithMaterialExpressionType::TextureCoordinate))
+		{
+			ConnectExpression(static_cast<const IDatasmithMaterialExpressionTextureCoordinate&>(*MaterialExpression), ConnectionData);
+		}
+		else if (MaterialExpression->IsSubType(EDatasmithMaterialExpressionType::FlattenNormal))
+		{
+			ConnectExpression(static_cast<const IDatasmithMaterialExpressionFlattenNormal&>(*MaterialExpression), ConnectionData);
+		}
+		else if (MaterialExpression->IsSubType(EDatasmithMaterialExpressionType::Generic))
+		{
+			ConnectExpression(static_cast<const IDatasmithMaterialExpressionGeneric&>(*MaterialExpression), ConnectionData, OutputIndex);
+		}
+		else if (MaterialExpression->IsSubType(EDatasmithMaterialExpressionType::FunctionCall))
+		{
+			ConnectExpression(static_cast<const IDatasmithMaterialExpressionFunctionCall&>(*MaterialExpression), ConnectionData, OutputIndex);
+		}
+		else if (MaterialExpression->IsSubType(EDatasmithMaterialExpressionType::Custom))
+		{
+			ConnectExpression(static_cast<const IDatasmithMaterialExpressionCustom&>(*MaterialExpression), ConnectionData, OutputIndex);
+		}
+	}
+
+	EBlendMode FPbrMaterialHelper::GetUEPbrImportBlendMode()
+	{
+		EBlendMode BlendMode = static_cast<EBlendMode>(MaterialElement.GetBlendMode());
+
+		if (MaterialElement.GetOpacity().GetExpression() && BlendMode != BLEND_Translucent && BlendMode != BLEND_Masked)
+		{
+			return EBlendMode::BLEND_Translucent; // force translucent
+		}
+
+		// If this material is dependent on function call expression, we need to check if those are translucent to apply
+		// translucency attribute to this material too.
+		for (int32 ExpressionIndex = 0; ExpressionIndex < MaterialElement.GetExpressionsCount(); ExpressionIndex++)
+		{
+			IDatasmithMaterialExpression* MaterialExpression = MaterialElement.GetExpression(ExpressionIndex);
+			if (MaterialExpression && MaterialExpression->IsSubType(EDatasmithMaterialExpressionType::FunctionCall))
+			{
+				IDatasmithMaterialExpressionFunctionCall* FunctionCall = static_cast<IDatasmithMaterialExpressionFunctionCall*>(MaterialExpression);
+
+				FString MaterialFunctionNodeUid = NodeUtils::MaterialFunctionPrefix + FunctionCall->GetFunctionPathName();
+				if (const UInterchangeDatasmithPbrMaterialNode* MaterialFunctionNode = Cast<const UInterchangeDatasmithPbrMaterialNode>(NodeContainer.GetNode(MaterialFunctionNodeUid)))
+				{
+					EBlendMode MaterialFunctionBlendMode;
+					if (MaterialFunctionNode->GetCustomBlendMode(MaterialFunctionBlendMode) && MaterialFunctionBlendMode == EBlendMode::BLEND_Translucent)
+					{
+						return MaterialFunctionBlendMode;
+					}
+				}
+			}
+		}
+
+		return BlendMode;
+	}
+#endif
+
+	bool BuildMaterialNode(IDatasmithUEPbrMaterialElement& MaterialElement, UInterchangeDatasmithPbrMaterialNode& MaterialNode)
+	{
+#if WITH_EDITOR
+		using namespace UE::Interchange::Materials;
+
+		FPbrMaterialHelper NodeHelper(MaterialNode, MaterialElement);
+
+		TFunction<void(IDatasmithExpressionInput&, FName)> ConnectOutputToInput;
+		ConnectOutputToInput = [&NodeHelper, &MaterialNode](const IDatasmithExpressionInput& ExpressionInput, FName InputName) -> void
+		{
+			FPbrMaterialHelper::FConnectionData ConnectionData{ &MaterialNode , InputName.ToString()};
+			NodeHelper.ConnectExpression(ExpressionInput.GetExpression(), ConnectionData, ExpressionInput.GetOutputIndex());
+		};
+
+		if (MaterialElement.GetUseMaterialAttributes())
+		{
+			//We ignore all the other inputs if we are using MaterialAttributes.
+			ConnectOutputToInput(MaterialElement.GetMaterialAttributes(), Common::Parameters::BxDF);
+		}
+		else
+		{
+			ConnectOutputToInput(MaterialElement.GetBaseColor(), PBR::Parameters::BaseColor);
+
+			ConnectOutputToInput(MaterialElement.GetMetallic(), PBR::Parameters::Metallic);
+
+			ConnectOutputToInput(MaterialElement.GetSpecular(), PBR::Parameters::Specular);
+
+			ConnectOutputToInput(MaterialElement.GetRoughness(), PBR::Parameters::Roughness);
+
+			ConnectOutputToInput(MaterialElement.GetEmissiveColor(), PBR::Parameters::EmissiveColor);
+
+			ConnectOutputToInput(MaterialElement.GetAmbientOcclusion(), PBR::Parameters::Occlusion);
+
+			ConnectOutputToInput(MaterialElement.GetNormal(), PBR::Parameters::Normal);
+
+			ConnectOutputToInput(MaterialElement.GetOpacity(), PBR::Parameters::Opacity);
+
+			ConnectOutputToInput(MaterialElement.GetRefraction(), PBR::Parameters::IndexOfRefraction);
+
+			ConnectOutputToInput(MaterialElement.GetClearCoat(), ClearCoat::Parameters::ClearCoat);
+
+			ConnectOutputToInput(MaterialElement.GetClearCoatRoughness(), ClearCoat::Parameters::ClearCoatRoughness);
+		}
+
+		for (const FString& MaterialFunctionUid : NodeHelper.MaterialFunctionUids)
+		{
+			MaterialNode.AddMaterialFunctionsDependency(MaterialFunctionUid);
+		}
+
+		MaterialNode.SetCustomTwoSided(MaterialElement.GetTwoSided());
+		MaterialNode.SetCustomOpacityMaskClipValue(MaterialElement.GetOpacityMaskClipValue());
+		MaterialNode.SetCustomTranslucencyLightingMode(static_cast<ETranslucencyLightingMode>(MaterialElement.GetTranslucencyLightingMode()));
+		MaterialNode.SetCustomBlendMode(NodeHelper.GetUEPbrImportBlendMode());
+
+		switch (MaterialElement.GetShadingModel())
+		{
+		case EDatasmithShadingModel::Subsurface:
+			MaterialNode.SetCustomShadingModel(MSM_Subsurface);
+			break;
+
+		case EDatasmithShadingModel::ClearCoat:
+			MaterialNode.SetCustomShadingModel(MSM_ClearCoat);
+			break;
+
+		case EDatasmithShadingModel::ThinTranslucent:
+			MaterialNode.SetCustomShadingModel(MSM_ThinTranslucent);
+			if (MaterialElement.GetOpacity().GetExpression())
+			{
+				MaterialNode.SetCustomTranslucencyLightingMode(ETranslucencyLightingMode::TLM_SurfacePerPixelLighting);
+			}
+			break;
+		case EDatasmithShadingModel::Unlit:
+			MaterialNode.SetCustomShadingModel(MSM_Unlit);
+			break;
+		default:
+			break;
+		}
+
+		if (MaterialElement.GetMaterialFunctionOnly())
+		{
+			MaterialNode.SetCustomIsAShaderFunction(true);
+		}
+
+		// Connect expressions to any UMaterialExpressionCustomOutput since these aren't part of the predefined material outputs
+		for (int32 ExpressionIndex = 0; ExpressionIndex < MaterialElement.GetExpressionsCount(); ++ExpressionIndex)
+		{
+			if (!NodeHelper.ProcessedExpressions[ExpressionIndex])
+			{
+				if (IDatasmithMaterialExpression* MaterialExpression = MaterialElement.GetExpression(ExpressionIndex))
+				{
+					// Unconnected generic expressions usually are custom output, connect them to the appropriate channel
+					if (MaterialExpression->IsSubType(EDatasmithMaterialExpressionType::Generic))
+					{
+						IDatasmithMaterialExpressionGeneric& GenericExpression = static_cast<IDatasmithMaterialExpressionGeneric&>(*MaterialExpression);
+						if (!FCString::Strcmp(GenericExpression.GetExpressionName(), TEXT("ThinTranslucentMaterialOutput")))
+						{
+							ensure(GenericExpression.GetInputCount() == 1);
+							ConnectOutputToInput(*GenericExpression.GetInput(0), ThinTranslucent::Parameters::TransmissionColor);
+						}
+						else if (!FCString::Strcmp(GenericExpression.GetExpressionName(), TEXT("ClearCoatNormalCustomOutput")))
+						{
+							ensure(GenericExpression.GetInputCount() == 1);
+							ConnectOutputToInput(*GenericExpression.GetInput(0), ClearCoat::Parameters::ClearCoatNormal);
+						}
+					}
+					else
+					{
+						ensure(MaterialExpression->IsSubType(EDatasmithMaterialExpressionType::ConstantBool) ||
+							MaterialExpression->IsSubType(EDatasmithMaterialExpressionType::ConstantColor) ||
+							MaterialExpression->IsSubType(EDatasmithMaterialExpressionType::ConstantScalar));
+					}
+				}
+			}
+		}
+
+		return true;
+#else
+		return false;
+#endif
+	}
+
+	bool BuildMaterialNode(IDatasmithDecalMaterialElement& MaterialElement, UInterchangeDatasmithMaterialNode& MaterialNode)
+	{
+		using namespace UE::Interchange::Materials;
+
+		return false;
+	}
+
+	void ProcessDependencies(IDatasmithUEPbrMaterialElement& MaterialElement, TArray<TSharedPtr<IDatasmithBaseMaterialElement>>& OutMaterialElements, TMap<FString, TSharedPtr<IDatasmithBaseMaterialElement>>& MaterialsToSort)
+	{
+		for (int32 ExpressionIndex = 0; ExpressionIndex < MaterialElement.GetExpressionsCount(); ExpressionIndex++)
+		{
+			IDatasmithMaterialExpression* MaterialExpression = MaterialElement.GetExpression(ExpressionIndex);
+			if (MaterialExpression && MaterialExpression->IsSubType(EDatasmithMaterialExpressionType::FunctionCall))
+			{
+				IDatasmithMaterialExpressionFunctionCall* MaterialExpressionFunctionCall = static_cast<IDatasmithMaterialExpressionFunctionCall*>(MaterialExpression);
+
+				TSharedPtr<IDatasmithBaseMaterialElement> MaterialFunction;
+				if (MaterialsToSort.RemoveAndCopyValue(MaterialExpressionFunctionCall->GetFunctionPathName(), MaterialFunction))
+				{
+					ProcessDependencies(static_cast<IDatasmithUEPbrMaterialElement&>(*MaterialFunction), OutMaterialElements, MaterialsToSort);
+
+					OutMaterialElements.Add(MaterialFunction);
+				}
+			}
+		}
+	}
+
+	// Preload in memory all the material functions referenced in the Pbr materials
+	// Rationale: The translator is not done in the game thread and some material function might not have been
+	//			  loaded yet. They need to be loaded during the translation in order to get the name of their inputs.
+	void PreLoadReferencedAssets(const TArray<TSharedPtr<IDatasmithBaseMaterialElement>>& MaterialElements)
+	{
+#if WITH_EDITOR
+		if (IsInGameThread())
+		{
+			return;
+		}
+
+		TSet<FString> ProcessedMaterialFunctions;
+		FGraphEventArray ThingsToComplete;
+		ThingsToComplete.Reserve(MaterialElements.Num());
+
+		for (const TSharedPtr<IDatasmithBaseMaterialElement>& MaterialElement : MaterialElements)
+		{
+			if (MaterialElement && MaterialElement->IsA(EDatasmithElementType::UEPbrMaterial))
+			{
+				IDatasmithUEPbrMaterialElement& PbrMaterialElement = static_cast<IDatasmithUEPbrMaterialElement&>(*MaterialElement);
+
+				for (int32 ExpressionIndex = 0; ExpressionIndex < PbrMaterialElement.GetExpressionsCount(); ExpressionIndex++)
+				{
+					IDatasmithMaterialExpression* MaterialExpression = PbrMaterialElement.GetExpression(ExpressionIndex);
+					if (MaterialExpression && MaterialExpression->IsSubType(EDatasmithMaterialExpressionType::FunctionCall))
+					{
+						IDatasmithMaterialExpressionFunctionCall* FunctionCall = static_cast<IDatasmithMaterialExpressionFunctionCall*>(MaterialExpression);
+
+						FString FunctionPath = FPackageName::ExportTextPathToObjectPath(FunctionCall->GetFunctionPathName());
+
+						if (!ProcessedMaterialFunctions.Contains(FunctionPath) && FPackageName::DoesPackageExist(FunctionPath))
+						{
+							ProcessedMaterialFunctions.Add(FunctionPath);
+
+							FSoftObjectPath ObjectPath(TEXT("MaterialFunction'") + FunctionPath + TEXT("'"));
+							if (!ObjectPath.ResolveObject())
+							{
+								TFunction<void()> Task = [FunctionPath]() -> void
+								{
+									FSoftObjectPath ObjectPath(TEXT("MaterialFunction'") + FunctionPath + TEXT("'"));
+									ensure(ObjectPath.TryLoad());
+								};
+
+								ThingsToComplete.Add(FFunctionGraphTask::CreateAndDispatchWhenReady(MoveTemp(Task), TStatId(), nullptr, ENamedThreads::GameThread));
+							}
+						}
+					}
+				}
+			}
+		}
+
+		if (ThingsToComplete.Num())
+		{
+			FTaskGraphInterface::Get().WaitUntilTasksComplete(ThingsToComplete, ENamedThreads::AnyThread);
+		}
+#endif
+	}
+
+	void ProcessMaterialElements(TArray<TSharedPtr<IDatasmithBaseMaterialElement>>& InOutMaterialElements)
+	{
+		PreLoadReferencedAssets(InOutMaterialElements);
+
+		bool bNeedSorting = false;
+
+		for (TSharedPtr<IDatasmithBaseMaterialElement>& MaterialElement : InOutMaterialElements)
+		{
+			if (MaterialElement->IsA(EDatasmithElementType::UEPbrMaterial))
+			{
+				bNeedSorting = true;
+				break;
+			}
+		}
+
+		if (!bNeedSorting)
+		{
+			return;
+		}
+
+		TMap<FString, TSharedPtr<IDatasmithBaseMaterialElement>> MaterialsToSort;
+		TArray<TSharedPtr<IDatasmithBaseMaterialElement>> MaterialElements(InOutMaterialElements);
+
+		MaterialsToSort.Reserve(MaterialElements.Num());
+		InOutMaterialElements.Empty(MaterialElements.Num());
+
+		for (TSharedPtr<IDatasmithBaseMaterialElement>& MaterialElement : MaterialElements)
+		{
+			if (MaterialElement->IsA(EDatasmithElementType::UEPbrMaterial))
+			{
+				MaterialsToSort.Add(MaterialElement->GetName(), MaterialElement);
+				continue;
+			}
+
+			InOutMaterialElements.Add(MaterialElement);
+		}
+
+		TArray<FString> MaterialNames;
+		MaterialsToSort.GenerateKeyArray(MaterialNames);
+
+		for (const FString& MaterialName : MaterialNames)
+		{
+			TSharedPtr<IDatasmithBaseMaterialElement> MaterialFunction;
+			if (MaterialsToSort.RemoveAndCopyValue(MaterialName, MaterialFunction))
+			{
+				ProcessDependencies(static_cast<IDatasmithUEPbrMaterialElement&>(*MaterialFunction), InOutMaterialElements, MaterialsToSort);
+
+				InOutMaterialElements.Add(MaterialFunction);
+			}
+		}
+	}
+
+	UInterchangeBaseNode* AddMaterialNode(const TSharedPtr<IDatasmithBaseMaterialElement>& MaterialElement, UInterchangeBaseNodeContainer& NodeContainer)
+	{
+		UInterchangeBaseNode* BaseNode = nullptr;
+
+		if (MaterialElement.IsValid())
+		{
+			if (MaterialElement->IsA(EDatasmithElementType::MaterialInstance))
+			{
+				UInterchangeDatasmithMaterialNode* MaterialNode = NewObject< UInterchangeDatasmithMaterialNode >(&NodeContainer);
+				const FString MaterialNodeUid = NodeUtils::MaterialPrefix + FDatasmithUtils::SanitizeObjectName(MaterialElement->GetName());
+				MaterialNode->InitializeNode(MaterialNodeUid, *FDatasmithUtils::SanitizeObjectName(MaterialElement->GetLabel()), EInterchangeNodeContainerType::TranslatedAsset);
+
+				if (!BuildMaterialNode(static_cast<IDatasmithMaterialInstanceElement&>(*MaterialElement), *MaterialNode))
+				{
+					MaterialNode = nullptr;
+				}
+
+				BaseNode = MaterialNode;
+			}
+			else if (MaterialElement->IsA(EDatasmithElementType::UEPbrMaterial))
+			{
+#if WITH_EDITOR
+				IDatasmithUEPbrMaterialElement& PbrMaterialElement = static_cast<IDatasmithUEPbrMaterialElement&>(*MaterialElement);
+
+				UInterchangeDatasmithPbrMaterialNode* MaterialNode = NewObject<UInterchangeDatasmithPbrMaterialNode>(&NodeContainer);
+
+				const FString& Prefix = PbrMaterialElement.GetMaterialFunctionOnly() ? NodeUtils::MaterialFunctionPrefix : NodeUtils::MaterialPrefix;
+				const FString MaterialNodeUid = Prefix + FDatasmithUtils::SanitizeObjectName(MaterialElement->GetName());
+				MaterialNode->InitializeNode(MaterialNodeUid, MaterialElement->GetLabel(), EInterchangeNodeContainerType::TranslatedAsset);
+
+				if (!BuildMaterialNode(PbrMaterialElement, *MaterialNode))
+				{
+					MaterialNode = nullptr;
+				}
+
+				BaseNode = MaterialNode;
+#endif
+			}
+
+			if (BaseNode)
+			{
+				BaseNode->SetAssetName(FDatasmithUtils::SanitizeObjectName(MaterialElement->GetName()));
+				NodeContainer.AddNode(BaseNode);
+			}
+		}
+
+		return BaseNode;
+	}
+}
+
diff --git a/Engine/Plugins/Experimental/Enterprise/DatasmithInterchange/Source/DatasmithInterchange/Private/InterchangeDatasmithMaterialNode.h b/Engine/Plugins/Experimental/Enterprise/DatasmithInterchange/Source/DatasmithInterchange/Private/InterchangeDatasmithMaterialNode.h
new file mode 100644
--- /dev/null
+++ b/Engine/Plugins/Experimental/Enterprise/DatasmithInterchange/Source/DatasmithInterchange/Private/InterchangeDatasmithMaterialNode.h
@@ -0,0 +1,189 @@
+// Copyright Epic Games, Inc. All Rights Reserved.
+
+#pragma once
+
+#include "CoreMinimal.h"
+
+#include "DatasmithDefinitions.h"
+
+#include "Nodes/InterchangeBaseNode.h"
+#include "InterchangeMaterialFactoryNode.h"
+#include "InterchangeShaderGraphNode.h"
+
+#include "Engine/EngineTypes.h"
+
+#include "InterchangeDatasmithMaterialNode.generated.h"
+
+class IDatasmithBaseMaterialElement;
+class UInterchangeBaseNodeContainer;
+
+UCLASS(BlueprintType, Experimental)
+class DATASMITHINTERCHANGE_API UInterchangeDatasmithMaterialNode : public UInterchangeBaseNode
+{
+	GENERATED_BODY()
+
+public:
+	/**
+	 * Return the node type name of the class, we use this when reporting error
+	 */
+	virtual FString GetTypeName() const override
+	{
+		const FString TypeName = TEXT("DatasmithMaterialNode");
+		return TypeName;
+	}
+
+	void SetMaterialType(EDatasmithReferenceMaterialType MaterialType)
+	{
+		AddInt32Attribute(MaterialTypeAttrName, (int32)MaterialType);
+	}
+
+	EDatasmithReferenceMaterialType GetMaterialType()
+	{
+		int32 AttributeValue;
+		return GetInt32Attribute(MaterialTypeAttrName, AttributeValue) ? (EDatasmithReferenceMaterialType)AttributeValue : EDatasmithReferenceMaterialType::Opaque;
+	}
+
+	void SetMaterialQuality(EDatasmithReferenceMaterialQuality MaterialQuality)
+	{
+		AddInt32Attribute(MaterialQualityAttrName, (int32)MaterialQuality);
+	}
+
+	EDatasmithReferenceMaterialQuality GetMaterialQuality()
+	{
+		int32 AttributeValue;
+		return GetInt32Attribute(MaterialQualityAttrName, AttributeValue) ? (EDatasmithReferenceMaterialQuality)AttributeValue : EDatasmithReferenceMaterialQuality::High;
+	}
+
+	void SetParentPath(const TCHAR* InSelectorName)
+	{
+		AddStringAttribute(MaterialParentAttrName, InSelectorName ? InSelectorName : FString());
+	}
+
+	FString GetParentPath()
+	{
+		FString AttributeValue;
+		return GetStringAttribute(MaterialParentAttrName, AttributeValue) ? AttributeValue : FString();
+	}
+
+public:
+	// TODO: See if it is not more appropriate to use IMPLEMENT_NODE_ATTRIBUTE_GETTER/SETTER defines
+	static const FName MaterialTypeAttrName;
+	static const FName MaterialQualityAttrName;
+	static const FName MaterialParentAttrName;
+};
+
+UCLASS(BlueprintType, Experimental)
+class DATASMITHINTERCHANGE_API UInterchangeDatasmithPbrMaterialNode : public UInterchangeShaderGraphNode
+{
+	GENERATED_BODY()
+
+public:
+
+	UInterchangeDatasmithPbrMaterialNode()
+	{
+		MaterialFunctionsDependencies.Initialize(Attributes, MaterialFunctionsDependenciesKey);
+	}
+
+	/**
+	 * Return the node type name of the class, we use this when reporting error
+	 */
+	virtual FString GetTypeName() const override
+	{
+		const FString TypeName = TEXT("DatasmithPbrMaterialNode");
+		return TypeName;
+	}
+
+	bool GetCustomShadingModel(EMaterialShadingModel& AttributeValue) const
+	{
+		int32 LocalAttributeValue;
+		if (GetInt32Attribute(ShadingModelAttrName, LocalAttributeValue))
+		{
+			AttributeValue = static_cast<EMaterialShadingModel>(LocalAttributeValue);
+			return true;
+		}
+
+		return false;
+	}
+
+	bool SetCustomShadingModel(EMaterialShadingModel AttributeValue)
+	{
+		return AddInt32Attribute(ShadingModelAttrName, static_cast<int32>(AttributeValue));
+	}
+
+	bool GetCustomBlendMode(EBlendMode& AttributeValue) const
+	{
+		int32 LocalAttributeValue;
+		if (GetInt32Attribute(BlendModeAttrName, LocalAttributeValue))
+		{
+			AttributeValue = static_cast<EBlendMode>(LocalAttributeValue);
+			return true;
+		}
+
+		return false;
+	}
+
+	bool SetCustomBlendMode(EBlendMode AttributeValue)
+	{
+		return AddInt32Attribute(BlendModeAttrName, static_cast<int32>(AttributeValue));
+	}
+
+	bool GetCustomOpacityMaskClipValue(float& AttributeValue) const
+	{
+		return GetFloatAttribute(OpacityMaskClipValueAttrName, AttributeValue);
+	}
+
+	bool SetCustomOpacityMaskClipValue(float AttributeValue)
+	{
+		return AddFloatAttribute(OpacityMaskClipValueAttrName, AttributeValue);
+	}
+
+	bool GetCustomTranslucencyLightingMode(ETranslucencyLightingMode& AttributeValue) const
+	{
+		int32 LocalAttributeValue;
+		if (GetInt32Attribute(TranslucencyLightingModeAttrName, LocalAttributeValue))
+		{
+			AttributeValue = static_cast<ETranslucencyLightingMode>(LocalAttributeValue);
+			return true;
+		}
+
+		return false;
+	}
+
+	bool SetCustomTranslucencyLightingMode(ETranslucencyLightingMode AttributeValue)
+	{
+		return AddInt32Attribute(TranslucencyLightingModeAttrName, static_cast<int32>(AttributeValue));
+	}
+
+	void GetMaterialFunctionsDependencies(TArray<FString>& MaterialFunctionUids)
+	{
+		MaterialFunctionsDependencies.GetItems(MaterialFunctionUids);
+	}
+
+	bool AddMaterialFunctionsDependency(const FString& MaterialFunctionUid)
+	{
+		return MaterialFunctionsDependencies.AddItem(MaterialFunctionUid);
+	}
+
+	bool RemoveAllMaterialFunctionsDependencies()
+	{
+		return MaterialFunctionsDependencies.RemoveAllItems();
+	}
+
+private:
+	static const FName ShadingModelAttrName;
+	static const FName BlendModeAttrName;
+	static const FName OpacityMaskClipValueAttrName;
+	static const FName TranslucencyLightingModeAttrName;
+	static const FString MaterialFunctionsDependenciesKey;
+
+	UE::Interchange::TArrayAttributeHelper<FString> MaterialFunctionsDependencies;
+};
+
+namespace UE::DatasmithInterchange::MaterialUtils
+{
+	extern const FName DefaultOutputIndexAttrName;
+	extern const FName MaterialFunctionPathAttrName;
+
+	void ProcessMaterialElements(TArray<TSharedPtr<IDatasmithBaseMaterialElement>>& InOutMaterialElements);
+	UInterchangeBaseNode* AddMaterialNode(const TSharedPtr<IDatasmithBaseMaterialElement>& MaterialElement, UInterchangeBaseNodeContainer& NodeContainer);
+}
diff --git a/Engine/Plugins/Experimental/Enterprise/DatasmithInterchange/Source/DatasmithInterchange/Private/InterchangeDatasmithMaterialPipeline.cpp b/Engine/Plugins/Experimental/Enterprise/DatasmithInterchange/Source/DatasmithInterchange/Private/InterchangeDatasmithMaterialPipeline.cpp
new file mode 100644
--- /dev/null
+++ b/Engine/Plugins/Experimental/Enterprise/DatasmithInterchange/Source/DatasmithInterchange/Private/InterchangeDatasmithMaterialPipeline.cpp
@@ -0,0 +1,344 @@
+// Copyright Epic Games, Inc. All Rights Reserved. 
+
+#include "InterchangeDatasmithMaterialPipeline.h"
+
+#include "InterchangeDatasmithMaterialNode.h"
+#include "InterchangeDatasmithMaterialPipeline.h"
+#include "InterchangeDatasmithUtils.h"
+
+#include "DatasmithMaterialElements.h"
+#include "IDatasmithSceneElements.h"
+
+#include "InterchangeReferenceMaterials/DatasmithReferenceMaterialManager.h"
+#include "InterchangeReferenceMaterials/DatasmithReferenceMaterialSelector.h"
+
+#include "InterchangeManager.h"
+#include "InterchangeMaterialDefinitions.h"
+#include "InterchangeMaterialFactoryNode.h"
+#include "InterchangeTexture2DNode.h"
+#include "InterchangeTexture2DFactoryNode.h"
+
+#include "Materials/MaterialFunction.h"
+#include "Materials/MaterialInstanceConstant.h"
+#include "Materials/MaterialInstanceDynamic.h"
+#include "Misc/PackageName.h"
+
+UInterchangeDatasmithMaterialPipeline::UInterchangeDatasmithMaterialPipeline()
+	: Super()
+{
+	MaterialImport = EInterchangeMaterialImportOption::ImportAsMaterials;
+}
+
+void UInterchangeDatasmithMaterialPipeline::ExecutePreImportPipeline(UInterchangeBaseNodeContainer* NodeContainer, const TArray<UInterchangeSourceData*>& InSourceDatas)
+{
+	using namespace UE::DatasmithInterchange;
+
+	TArray< UInterchangeDatasmithMaterialNode*> ReferenceMaterials;
+	TArray<UInterchangeShaderNode*> ShaderNodes;
+
+	//Find all translated node we need for this pipeline
+	NodeContainer->IterateNodes([&](const FString& NodeUid, UInterchangeBaseNode* Node)
+		{
+			if (UInterchangeDatasmithMaterialNode* MaterialNode = Cast<UInterchangeDatasmithMaterialNode>(Node))
+			{
+				ReferenceMaterials.Add(MaterialNode);
+			}
+			else if (UInterchangeShaderNode* ShaderNode = Cast<UInterchangeShaderNode>(Node))
+			{
+				ShaderNodes.Add(ShaderNode);
+			}
+		});
+
+	Super::ExecutePreImportPipeline(NodeContainer, InSourceDatas);
+
+	UpdateMaterialFactoryNodes(ShaderNodes);
+
+	for (UInterchangeDatasmithMaterialNode* MaterialNode : ReferenceMaterials)
+	{
+		PreImportMaterialNode(NodeContainer, MaterialNode);
+	}
+}
+
+void UInterchangeDatasmithMaterialPipeline::ExecutePostImportPipeline(const UInterchangeBaseNodeContainer* NodeContainer, const FString& NodeKey, UObject* CreatedAsset, bool bIsAReimport)
+{
+	if (UMaterialInterface* ImportedMaterial = Cast<UMaterialInterface>(CreatedAsset))
+	{
+		UInterchangeMaterialInstanceFactoryNode* InstanceFactoryNode = nullptr;
+		NodeContainer->BreakableIterateNodesOfType<UInterchangeMaterialInstanceFactoryNode>([&ImportedMaterial, &InstanceFactoryNode](const FString& NodeUid, UInterchangeMaterialInstanceFactoryNode* Node)
+		{
+			FSoftObjectPath NodeReferenceObject;
+			Node->GetCustomReferenceObject(NodeReferenceObject);
+			if (NodeReferenceObject.TryLoad() == ImportedMaterial)
+			{
+				InstanceFactoryNode = Node;
+				return true;
+			}
+
+			return false;
+		});
+
+		if (InstanceFactoryNode)
+		{
+			PostImportMaterialInstanceFactoryNode(NodeContainer, InstanceFactoryNode, ImportedMaterial);
+		}
+
+		return;
+	}
+}
+
+void UInterchangeDatasmithMaterialPipeline::PreImportMaterialNode(UInterchangeBaseNodeContainer* NodeContainer, UInterchangeDatasmithMaterialNode* MaterialNode)
+{
+	FString ParentPath = MaterialNode->GetParentPath();
+
+	if (ParentPath.IsEmpty())
+	{
+		return;
+	}
+
+	using namespace UE::DatasmithInterchange;
+
+	FString PackagePath;
+
+	if (MaterialNode->GetMaterialType() == EDatasmithReferenceMaterialType::Custom && FPackageName::DoesPackageExist(ParentPath))
+	{
+		PackagePath = ParentPath;
+	}
+	else
+	{
+		FString Host(FDatasmithReferenceMaterialManager::Get().GetHostFromString(ParentPath));
+
+		if (TSharedPtr<FDatasmithReferenceMaterialSelector> MaterialSelector = FDatasmithReferenceMaterialManager::Get().GetSelector(*Host))
+		{
+			PackagePath = MaterialSelector->GetMaterialPath(MaterialNode->GetMaterialType());
+		}
+		else
+		{
+			//FText FailReason = FText::Format(LOCTEXT("NoSelectorForHost", "No Material selector found for Host {0}. Skipping material {1} ..."), FText::FromString(Host), FText::FromString(MaterialElement.GetName()));
+			//ImportContext.LogError(FailReason);
+		}
+	}
+
+	if (PackagePath.IsEmpty())
+	{
+		return;
+	}
+
+	const FString MaterialNodeUid = UInterchangeMaterialInstanceFactoryNode::GetMaterialFactoryNodeUidFromMaterialNodeUid(MaterialNode->GetUniqueID());
+
+	UInterchangeMaterialInstanceFactoryNode* MaterialFactoryNode = NodeUtils::FindOrAddFactoryNode<UInterchangeMaterialInstanceFactoryNode>(MaterialNode, NodeContainer, MaterialNodeUid);
+
+	// Set MaterialFactoryNode's display label to MaterialNode's uniqueID
+	// to reconcile mesh's slot names and material assets
+	MaterialFactoryNode->SetDisplayLabel(MaterialNode->GetAssetName());
+	MaterialFactoryNode->SetCustomParent(PackagePath);
+
+#if WITH_EDITOR
+	MaterialFactoryNode->SetCustomInstanceClassName(UMaterialInstanceConstant::StaticClass()->GetPathName());
+#else
+	MaterialFactoryNode->SetCustomInstanceClassName(UMaterialInstanceDynamic::StaticClass()->GetPathName());
+#endif
+
+	TArray<FString> Inputs;
+	UInterchangeShaderPortsAPI::GatherInputs(MaterialNode, Inputs);
+
+	for (const FString& InputName : Inputs)
+	{
+		FName InputValueKey = UInterchangeShaderPortsAPI::MakeInputValueKey(InputName);
+
+		switch (UInterchangeShaderPortsAPI::GetInputType(MaterialNode, InputName))
+		{
+			case UE::Interchange::EAttributeTypes::Bool:
+			{
+				bool AttributeValue = false;
+				MaterialNode->GetBooleanAttribute(InputValueKey, AttributeValue);
+				MaterialFactoryNode->AddBooleanAttribute(InputValueKey, AttributeValue);
+			}
+			break;
+			case UE::Interchange::EAttributeTypes::Int32:
+			{
+				int32 AttributeValue = 0;
+				MaterialNode->GetInt32Attribute(InputValueKey, AttributeValue);
+				MaterialFactoryNode->AddInt32Attribute(InputValueKey, AttributeValue);
+			}
+			break;
+			case UE::Interchange::EAttributeTypes::Float:
+			{
+				float AttributeValue = 0.f;
+				MaterialNode->GetFloatAttribute(InputValueKey, AttributeValue);
+				MaterialFactoryNode->AddFloatAttribute(InputValueKey, AttributeValue);
+			}
+			break;
+			case UE::Interchange::EAttributeTypes::LinearColor:
+			{
+				FLinearColor AttributeValue = FLinearColor::White;
+				MaterialNode->GetLinearColorAttribute(InputValueKey, AttributeValue);
+				MaterialFactoryNode->AddLinearColorAttribute(InputValueKey, AttributeValue);
+			}
+			break;
+			case UE::Interchange::EAttributeTypes::String:
+			{
+				FString TextureName;
+				MaterialNode->GetStringAttribute(InputValueKey, TextureName);
+
+				if (!FPackageName::IsValidObjectPath(TextureName))
+				{
+					const FString TextureUid = UInterchangeTexture2DFactoryNode::GetTextureFactoryNodeUidFromTextureNodeUid(NodeUtils::TexturePrefix + TextureName);
+					MaterialFactoryNode->AddStringAttribute(InputValueKey, TextureUid);
+					MaterialFactoryNode->AddFactoryDependencyUid(TextureUid);
+				}
+				else
+				{
+					MaterialFactoryNode->AddStringAttribute(InputValueKey, TextureName);
+				}
+			}
+			break;
+		}
+	}
+
+	if (MaterialNode->GetMaterialType() != EDatasmithReferenceMaterialType::Custom)
+	{
+		MaterialFactoryNode->AddStringAttribute(UInterchangeDatasmithMaterialNode::MaterialParentAttrName, MaterialNode->GetParentPath());
+		MaterialFactoryNode->AddInt32Attribute(UInterchangeDatasmithMaterialNode::MaterialTypeAttrName, (int32)MaterialNode->GetMaterialType());
+		MaterialFactoryNode->AddInt32Attribute(UInterchangeDatasmithMaterialNode::MaterialQualityAttrName, (int32)MaterialNode->GetMaterialQuality());
+	}
+}
+
+void UInterchangeDatasmithMaterialPipeline::PostImportMaterialInstanceFactoryNode(const UInterchangeBaseNodeContainer* NodeContainer, UInterchangeMaterialInstanceFactoryNode* FactoryNode, UMaterialInterface* CreatedMaterial)
+{
+	using namespace UE::DatasmithInterchange;
+
+#if WITH_EDITOR
+	FString SelectorName;
+	if (FactoryNode->GetStringAttribute(UInterchangeDatasmithMaterialNode::MaterialParentAttrName, SelectorName))
+	{
+		if (!SelectorName.IsEmpty())
+		{
+			const FString Host(FDatasmithReferenceMaterialManager::Get().GetHostFromString(SelectorName));
+
+			int32 MaterialType;
+			int32 MaterialQuality;
+
+			ensure(FactoryNode->GetInt32Attribute(UInterchangeDatasmithMaterialNode::MaterialTypeAttrName, MaterialType));
+			ensure(FactoryNode->GetInt32Attribute(UInterchangeDatasmithMaterialNode::MaterialQualityAttrName, MaterialQuality));
+
+			if (TSharedPtr<FDatasmithReferenceMaterialSelector> MaterialSelector = FDatasmithReferenceMaterialManager::Get().GetSelector(*Host))
+			{
+				MaterialSelector->PostImportProcess((EDatasmithReferenceMaterialType)MaterialType, (EDatasmithReferenceMaterialQuality)MaterialQuality, Cast<UMaterialInstanceConstant>(CreatedMaterial));
+			}
+		}
+	}
+#endif
+
+	// Process parameters which are connected to pre-existing textures
+	TArray<FString> Inputs;
+	UInterchangeShaderPortsAPI::GatherInputs(FactoryNode, Inputs);
+
+	for (const FString& InputName : Inputs)
+	{
+		FName InputValueKey = UInterchangeShaderPortsAPI::MakeInputValueKey(InputName);
+
+		if (UInterchangeShaderPortsAPI::GetInputType(FactoryNode, InputName) == UE::Interchange::EAttributeTypes::String)
+		{
+			const FName ParameterName = *InputName;
+			UTexture* InstanceValue = nullptr;
+
+			if (CreatedMaterial->GetTextureParameterValue(ParameterName, InstanceValue))
+			{
+				FString TextureName;
+				FactoryNode->GetStringAttribute(InputValueKey, TextureName);
+
+				if (FPackageName::IsValidObjectPath(TextureName))
+				{
+					if (UTexture* InputTexture = Cast<UTexture>(FSoftObjectPath(TextureName).TryLoad()))
+					{
+						if (InputTexture != InstanceValue)
+						{
+#if WITH_EDITOR
+							if (UMaterialInstanceConstant* MaterialInstanceConstant = Cast<UMaterialInstanceConstant>(CreatedMaterial))
+							{
+								MaterialInstanceConstant->SetTextureParameterValueEditorOnly(ParameterName, InputTexture);
+							}
+							else
+#endif // WITH_EDITOR
+							if (UMaterialInstanceDynamic* MaterialInstanceDynamic = Cast<UMaterialInstanceDynamic>(CreatedMaterial))
+							{
+								MaterialInstanceDynamic->SetTextureParameterValue(ParameterName, InputTexture);
+							}
+						}
+					}
+				}
+			}
+		}
+	}
+}
+
+void UInterchangeDatasmithMaterialPipeline::UpdateMaterialFactoryNodes(const TArray<UInterchangeShaderNode*>& ShaderNodes)
+{
+	using namespace UE::DatasmithInterchange;
+
+	for (UInterchangeShaderNode* ShaderNode : ShaderNodes)
+	{
+		const FString FactoryNodeUid = UInterchangeMaterialFactoryNode::GetMaterialFactoryNodeUidFromMaterialNodeUid(ShaderNode->GetUniqueID());
+		UInterchangeFactoryBaseNode* FactoryNode = BaseNodeContainer->GetFactoryNode(FactoryNodeUid);
+		if (!ensure(FactoryNode))
+		{
+			continue;
+		}
+
+		FString MaterialFunctionPath;
+		if (ShaderNode->GetStringAttribute(MaterialUtils::MaterialFunctionPathAttrName, MaterialFunctionPath))
+		{
+			static const FName MaterialFunctionMemberName = GET_MEMBER_NAME_CHECKED(UMaterialExpressionMaterialFunctionCall, MaterialFunction);
+
+			if (UInterchangeMaterialExpressionFactoryNode* ExpressionFactoryNode = Cast<UInterchangeMaterialExpressionFactoryNode>(FactoryNode))
+			{
+				ExpressionFactoryNode->AddStringAttribute(MaterialFunctionMemberName, MaterialFunctionPath);
+				ExpressionFactoryNode->AddApplyAndFillDelegates<FString>(MaterialFunctionMemberName, UMaterialExpressionMaterialFunctionCall::StaticClass(), MaterialFunctionMemberName);
+			}
+		}
+		else if (UInterchangeDatasmithPbrMaterialNode* PbrMaterialNode = Cast<UInterchangeDatasmithPbrMaterialNode>(ShaderNode))
+		{
+			UInterchangeBaseMaterialFactoryNode* BaseMaterialFactoryNode = Cast<UInterchangeBaseMaterialFactoryNode>(FactoryNode);
+			ensure(BaseMaterialFactoryNode);
+
+			TArray<FString> MaterialFunctionUids;
+			PbrMaterialNode->GetMaterialFunctionsDependencies(MaterialFunctionUids);
+			for (const FString& MaterialFunctionUid : MaterialFunctionUids)
+			{
+				const FString MaterialFunctionFactoryNodeUid = UInterchangeMaterialFactoryNode::GetMaterialFactoryNodeUidFromMaterialNodeUid(MaterialFunctionUid);
+				ensure(BaseNodeContainer->GetFactoryNode(MaterialFunctionFactoryNodeUid));
+				BaseMaterialFactoryNode->AddFactoryDependencyUid(MaterialFunctionFactoryNodeUid);
+			}
+
+			PbrMaterialNode->RemoveAllMaterialFunctionsDependencies();
+
+			if (UInterchangeMaterialFactoryNode* MaterialFactoryNode = Cast<UInterchangeMaterialFactoryNode>(FactoryNode))
+			{
+				EMaterialShadingModel ShadingModel;
+				if (PbrMaterialNode->GetCustomShadingModel(ShadingModel))
+				{
+					MaterialFactoryNode->SetCustomShadingModel(ShadingModel);
+				}
+
+				EBlendMode BlendMode;
+				if (PbrMaterialNode->GetCustomBlendMode(BlendMode))
+				{
+					MaterialFactoryNode->SetCustomBlendMode(BlendMode);
+				}
+
+				ETranslucencyLightingMode TranslucencyLightingMode;
+				if (PbrMaterialNode->GetCustomTranslucencyLightingMode(TranslucencyLightingMode))
+				{
+					MaterialFactoryNode->SetCustomTranslucencyLightingMode(TranslucencyLightingMode);
+				}
+
+				// TODO: Add opacity mask clip value when implemented
+				//float OpacityMaskClip;
+				//if (PbrMaterialNode->GetCustomOpacityMaskClipValue(OpacityMaskClip))
+				//{
+				//	MaterialFactoryNode->SetCustomOpacityMaskClipValue(OpacityMaskClip);
+				//}
+			}
+		}
+	}
+}
diff --git a/Engine/Plugins/Experimental/Enterprise/DatasmithInterchange/Source/DatasmithInterchange/Private/InterchangeDatasmithMaterialPipeline.h b/Engine/Plugins/Experimental/Enterprise/DatasmithInterchange/Source/DatasmithInterchange/Private/InterchangeDatasmithMaterialPipeline.h
new file mode 100644
--- /dev/null
+++ b/Engine/Plugins/Experimental/Enterprise/DatasmithInterchange/Source/DatasmithInterchange/Private/InterchangeDatasmithMaterialPipeline.h
@@ -0,0 +1,40 @@
+// Copyright Epic Games, Inc. All Rights Reserved.
+
+#pragma once
+
+#include "CoreMinimal.h"
+
+#include "InterchangeGenericMaterialPipeline.h"
+
+#include "InterchangeDatasmithMaterialPipeline.generated.h"
+
+class UMaterialInterface;
+class UInterchangeDatasmithMaterialNode;
+class UInterchangeDatasmithPbrMaterialNode;
+class UInterchangeFactoryBaseNode;
+class UInterchangeMaterialFactoryNode;
+class UInterchangeMaterialInstanceFactoryNode;
+class UInterchangeShaderGraphNode;
+
+UCLASS(BlueprintType, Experimental)
+class DATASMITHINTERCHANGE_API UInterchangeDatasmithMaterialPipeline : public UInterchangeGenericMaterialPipeline
+{
+	GENERATED_BODY()
+
+	UInterchangeDatasmithMaterialPipeline();
+
+protected:
+	virtual void ExecutePreImportPipeline(UInterchangeBaseNodeContainer* BaseNodeContainer, const TArray<UInterchangeSourceData*>& SourceDatas) override;
+	virtual void ExecutePostImportPipeline(const UInterchangeBaseNodeContainer* BaseNodeContainer, const FString& NodeKey, UObject* CreatedAsset, bool bIsAReimport) override;
+
+	virtual bool CanExecuteOnAnyThread(EInterchangePipelineTask PipelineTask) override
+	{
+		// This pipeline creates UObjects and assets. Not safe to execute outside of main thread.
+		return false;
+	}
+
+private:
+	void PreImportMaterialNode(UInterchangeBaseNodeContainer* NodeContainer, UInterchangeDatasmithMaterialNode* MaterialNode);
+	void PostImportMaterialInstanceFactoryNode(const UInterchangeBaseNodeContainer* NodeContainer, UInterchangeMaterialInstanceFactoryNode* FactoryNode, UMaterialInterface* CreatedMaterial);
+	void UpdateMaterialFactoryNodes(const TArray<UInterchangeShaderNode*>& ShaderNodes);
+};
\ No newline at end of file
diff --git a/Engine/Plugins/Experimental/Enterprise/DatasmithInterchange/Source/DatasmithInterchange/Private/InterchangeDatasmithPipeline.cpp b/Engine/Plugins/Experimental/Enterprise/DatasmithInterchange/Source/DatasmithInterchange/Private/InterchangeDatasmithPipeline.cpp
new file mode 100644
--- /dev/null
+++ b/Engine/Plugins/Experimental/Enterprise/DatasmithInterchange/Source/DatasmithInterchange/Private/InterchangeDatasmithPipeline.cpp
@@ -0,0 +1,229 @@
+// Copyright Epic Games, Inc. All Rights Reserved.
+
+#include "InterchangeDatasmithPipeline.h"
+
+#include "InterchangeDatasmithAreaLightNode.h"
+#include "InterchangeDatasmithAreaLightFactoryNode.h"
+#include "InterchangeDatasmithLevelPipeline.h"
+#include "InterchangeDatasmithMaterialPipeline.h"
+#include "InterchangeDatasmithSceneNode.h"
+#include "InterchangeDatasmithSceneFactoryNode.h"
+#include "InterchangeDatasmithStaticMeshPipeline.h"
+#include "InterchangeDatasmithTexturePipeline.h"
+#include "InterchangeDatasmithUtils.h"
+
+#include "InterchangeAnimationTrackSetFactoryNode.h"
+#include "InterchangeGenericAnimationPipeline.h"
+#include "InterchangeGenericScenesPipeline.h"
+#include "InterchangeMaterialFactoryNode.h"
+#include "InterchangeSceneVariantSetsFactoryNode.h"
+#include "InterchangeStaticMeshFactoryNode.h"
+#include "InterchangeTexture2DNode.h"
+#include "InterchangeTextureFactoryNode.h"
+#include "InterchangeTexture2DFactoryNode.h"
+
+#include "ExternalSource.h"
+#include "DatasmithAreaLightActor.h"
+#include "DatasmithScene.h"
+#include "DatasmithSceneXmlWriter.h"
+
+#if WITH_EDITOR
+#include "DatasmithImporter.h"
+#include "DatasmithImportContext.h"
+#include "DatasmithStaticMeshImporter.h"
+#include "Misc/ScopedSlowTask.h"
+#include "Utility/DatasmithImporterUtils.h"
+#endif //WITH_EDITOR
+
+#define LOCTEXT_NAMESPACE "InterchangeDatasmithPipeline"
+
+UInterchangeDatasmithPipeline::UInterchangeDatasmithPipeline()
+{
+	CommonSkeletalMeshesAndAnimationsProperties = CreateDefaultSubobject<UInterchangeGenericCommonSkeletalMeshesAndAnimationsProperties>("CommonSkeletalMeshesAndAnimationsProperties");
+	CommonSkeletalMeshesAndAnimationsProperties->bImportMeshesInBoneHierarchy = false;
+	
+	CommonMeshesProperties = CreateDefaultSubobject<UInterchangeGenericCommonMeshesProperties>("CommonMeshesProperties");
+	CommonMeshesProperties->bBakeMeshes = false;
+
+	TexturePipeline = CreateDefaultSubobject<UInterchangeDatasmithTexturePipeline>("DatasmithTexturePipeline");
+	MaterialPipeline = CreateDefaultSubobject<UInterchangeDatasmithMaterialPipeline>("DatasmithMaterialPipeline");
+	MeshPipeline = CreateDefaultSubobject<UInterchangeDatasmithStaticMeshPipeline>("DatasmithMeshPipeline");
+	LevelPipeline = CreateDefaultSubobject<UInterchangeDatasmithLevelPipeline>("DatasmithLevelPipeline");
+	AnimationPipeline = CreateDefaultSubobject<UInterchangeGenericAnimationPipeline>("AnimationPipeline");
+
+	MeshPipeline->CommonMeshesProperties = CommonMeshesProperties;
+	MeshPipeline->CommonSkeletalMeshesAndAnimationsProperties = CommonSkeletalMeshesAndAnimationsProperties;
+}
+
+void UInterchangeDatasmithPipeline::ExecutePreImportPipeline(UInterchangeBaseNodeContainer* InBaseNodeContainer, const TArray<UInterchangeSourceData*>& SourceDatas)
+{
+	using namespace UE::DatasmithInterchange;
+
+	BaseNodeContainer = InBaseNodeContainer;
+
+	TArray<UInterchangeDatasmithSceneNode*> DatasmithSceneNodes = NodeUtils::GetNodes<UInterchangeDatasmithSceneNode>(BaseNodeContainer);
+	if (!ensure(DatasmithSceneNodes.Num() == 1))
+	{
+		return;
+	}
+
+	if (TexturePipeline)
+	{
+		TexturePipeline->ScriptedExecutePreImportPipeline(BaseNodeContainer, SourceDatas);
+	}
+
+	if (MaterialPipeline)
+	{
+		MaterialPipeline->ScriptedExecutePreImportPipeline(BaseNodeContainer, SourceDatas);
+	}
+
+	if (MeshPipeline)
+	{
+		MeshPipeline->ScriptedExecutePreImportPipeline(BaseNodeContainer, SourceDatas);
+	}
+
+	if (LevelPipeline)
+	{
+		LevelPipeline->ScriptedExecutePreImportPipeline(BaseNodeContainer, SourceDatas);
+	}
+
+	if (AnimationPipeline)
+	{
+		AnimationPipeline->ScriptedExecutePreImportPipeline(BaseNodeContainer, SourceDatas);
+	}
+
+	const UInterchangeDatasmithSceneNode* DatasmithSceneNode = DatasmithSceneNodes[0];
+	const FString PackageSubPath = DatasmithSceneNode->GetDisplayLabel();
+
+	TArray<FString> DependenciesUids;
+
+	// Textures
+	for (UInterchangeTextureFactoryNode* TextureFactoryNode : NodeUtils::GetNodes<UInterchangeTextureFactoryNode>(BaseNodeContainer))
+	{
+		DependenciesUids.Add(TextureFactoryNode->GetUniqueID());
+		TextureFactoryNode->SetCustomSubPath(FPaths::Combine(PackageSubPath, "Textures"));
+		TextureFactoryNode->SetEnabled(true);
+	}
+
+	// Materials
+	for (UInterchangeBaseMaterialFactoryNode* MaterialFactoryNode : NodeUtils::GetNodes<UInterchangeBaseMaterialFactoryNode>(BaseNodeContainer))
+	{
+		DependenciesUids.Add(MaterialFactoryNode->GetUniqueID());
+		MaterialFactoryNode->SetCustomSubPath(FPaths::Combine(PackageSubPath, "Materials"));
+		MaterialFactoryNode->SetEnabled(true);
+	}
+
+	// StaticMeshes
+	for (UInterchangeStaticMeshFactoryNode* StaticMeshFactoryNode : NodeUtils::GetNodes<UInterchangeStaticMeshFactoryNode>(BaseNodeContainer))
+	{
+		DependenciesUids.Add(StaticMeshFactoryNode->GetUniqueID());
+		StaticMeshFactoryNode->SetCustomSubPath(FPaths::Combine(PackageSubPath, "Geometries"));
+		StaticMeshFactoryNode->SetEnabled(true);
+	}
+
+	// LevelSequences
+	for (UInterchangeAnimationTrackSetFactoryNode* AnimationTrackSetFactoryNode : NodeUtils::GetNodes<UInterchangeAnimationTrackSetFactoryNode>(BaseNodeContainer))
+	{
+		DependenciesUids.Add(AnimationTrackSetFactoryNode->GetUniqueID());
+		AnimationTrackSetFactoryNode->SetCustomSubPath(FPaths::Combine(PackageSubPath, "Animations"));
+		AnimationTrackSetFactoryNode->SetEnabled(true);
+	}
+
+	// LevelVariantSets
+	for (UInterchangeSceneVariantSetsFactoryNode* LevelVariantSetFactoryNode : NodeUtils::GetNodes<UInterchangeSceneVariantSetsFactoryNode>(BaseNodeContainer))
+	{
+		DependenciesUids.Add(LevelVariantSetFactoryNode->GetUniqueID());
+		LevelVariantSetFactoryNode->SetCustomSubPath(FPaths::Combine(PackageSubPath, "Variants"));
+		LevelVariantSetFactoryNode->SetEnabled(true);
+	}
+
+	// Datasmith Scene
+	{
+		const FString DatasmithSceneUid = TEXT("Factory_") + DatasmithSceneNode->GetUniqueID();
+		const FString DisplayLabel = DatasmithSceneNode->GetDisplayLabel();
+		UInterchangeDatasmithSceneFactoryNode* DatasmithSceneFactoryNode = NewObject<UInterchangeDatasmithSceneFactoryNode>(BaseNodeContainer, NAME_None);
+		if (!ensure(DatasmithSceneFactoryNode))
+		{
+			return;
+		}
+
+		DatasmithSceneFactoryNode->InitializeDatasmithFactorySceneNode(DatasmithSceneUid, DisplayLabel, UDatasmithScene::StaticClass()->GetName());
+		DatasmithSceneFactoryNode->SetCustomSubPath(PackageSubPath);
+		DatasmithSceneFactoryNode->AddTargetNodeUid(DatasmithSceneNode->GetUniqueID());
+		DatasmithSceneNode->AddTargetNodeUid(DatasmithSceneFactoryNode->GetUniqueID());
+
+		for (const FString& Uid : DependenciesUids)
+		{
+			DatasmithSceneFactoryNode->AddFactoryDependencyUid(Uid);
+		}
+
+		BaseNodeContainer->AddNode(DatasmithSceneFactoryNode);
+	}
+}
+
+void UInterchangeDatasmithPipeline::ExecutePostImportPipeline(const UInterchangeBaseNodeContainer* InBaseNodeContainer, const FString& NodeKey, UObject* CreatedAsset, bool bIsAReimport)
+{
+	TRACE_CPUPROFILER_EVENT_SCOPE(UInterchangeDatasmithPipeline::ExecutePostImportPipeline);
+
+	if (TexturePipeline)
+	{
+		TexturePipeline->ScriptedExecutePostImportPipeline(InBaseNodeContainer, NodeKey, CreatedAsset, bIsAReimport);
+	}
+
+	if (MaterialPipeline)
+	{
+		MaterialPipeline->ScriptedExecutePostImportPipeline(InBaseNodeContainer, NodeKey, CreatedAsset, bIsAReimport);
+	}
+
+	if (MeshPipeline)
+	{
+		MeshPipeline->ScriptedExecutePostImportPipeline(InBaseNodeContainer, NodeKey, CreatedAsset, bIsAReimport);
+	}
+
+	if (LevelPipeline)
+	{
+		LevelPipeline->ScriptedExecutePostImportPipeline(InBaseNodeContainer, NodeKey, CreatedAsset, bIsAReimport);
+	}
+
+	if (AnimationPipeline)
+	{
+		AnimationPipeline->ScriptedExecutePostImportPipeline(InBaseNodeContainer, NodeKey, CreatedAsset, bIsAReimport);
+	}
+
+	if (UDatasmithScene* DatasmithSceneAsset = Cast<UDatasmithScene>(CreatedAsset))
+	{
+		PostImportDatasmithSceneAsset(*DatasmithSceneAsset);
+	}
+}
+
+void UInterchangeDatasmithPipeline::PostImportDatasmithSceneAsset(UDatasmithScene& DatasmithSceneAsset)
+{
+	using namespace UE::DatasmithInterchange;
+
+#if WITH_EDITORONLY_DATA
+	TArray<UInterchangeDatasmithSceneNode*> DatasmithSceneNodes = NodeUtils::GetNodes<UInterchangeDatasmithSceneNode>(BaseNodeContainer);
+	if (!ensure(DatasmithSceneNodes.Num() == 1))
+	{
+		// TODO: Warn more than one Datasmith scene asset
+		return;
+	}
+
+	TArray< uint8 > Bytes;
+	FMemoryWriter MemoryWriter(Bytes, true);
+
+	FDatasmithSceneXmlWriter DatasmithSceneXmlWriter;
+	DatasmithSceneXmlWriter.Serialize(DatasmithSceneNodes[0]->DatasmithScene.ToSharedRef(), MemoryWriter);
+
+	DatasmithSceneAsset.DatasmithSceneBulkData.Lock(LOCK_READ_WRITE);
+
+	uint8* Dest = reinterpret_cast<uint8*>(DatasmithSceneAsset.DatasmithSceneBulkData.Realloc(Bytes.Num()));
+
+	FPlatformMemory::Memcpy(Dest, Bytes.GetData(), Bytes.Num());
+
+	DatasmithSceneAsset.DatasmithSceneBulkData.Unlock();
+
+	// Todo: Fill up imported DatasmithScene with created assets
+#endif
+}
+
+#undef LOCTEXT_NAMESPACE
\ No newline at end of file
diff --git a/Engine/Plugins/Experimental/Enterprise/DatasmithInterchange/Source/DatasmithInterchange/Private/InterchangeDatasmithSceneFactory.cpp b/Engine/Plugins/Experimental/Enterprise/DatasmithInterchange/Source/DatasmithInterchange/Private/InterchangeDatasmithSceneFactory.cpp
new file mode 100644
--- /dev/null
+++ b/Engine/Plugins/Experimental/Enterprise/DatasmithInterchange/Source/DatasmithInterchange/Private/InterchangeDatasmithSceneFactory.cpp
@@ -0,0 +1,125 @@
+// Copyright Epic Games, Inc. All Rights Reserved.
+
+#include "InterchangeDatasmithSceneFactory.h"
+
+#include "InterchangeDatasmithLog.h"
+#include "InterchangeDatasmithSceneFactoryNode.h"
+
+#include "DatasmithScene.h"
+#include "Logging/LogMacros.h"
+#include "Nodes/InterchangeBaseNode.h"
+
+UClass* UInterchangeDatasmithSceneFactory::GetFactoryClass() const
+{
+	return UDatasmithScene::StaticClass();
+}
+
+
+UObject* UInterchangeDatasmithSceneFactory::CreateEmptyAsset(const FCreateAssetParams& Arguments)
+{
+	UDatasmithScene* DatasmithScene = nullptr;
+
+#if WITH_EDITORONLY_DATA
+	if (!Arguments.AssetNode)
+	{
+		UE_LOG(LogInterchangeDatasmith, Error, TEXT("UInterchangeDatasmithSceneFactory: Asset node parameter is null."));
+		return nullptr;
+	}
+
+	const UClass* DatasmithSceneClass = Arguments.AssetNode->GetObjectClass();
+	if (!DatasmithSceneClass || !DatasmithSceneClass->IsChildOf(UDatasmithScene::StaticClass()))
+	{
+		UE_LOG(LogInterchangeDatasmith, Error, TEXT("UInterchangeDatasmithSceneFactory: Asset node parameter class doesnt derive from UDatasmithScene."));
+		return nullptr;
+	}
+
+	// create an asset if it doesn't exist
+	UObject* ExistingAsset = StaticFindObject(nullptr, Arguments.Parent, *Arguments.AssetName);
+
+	// create a new texture or overwrite existing asset, if possible
+	if (!ExistingAsset)
+	{
+		DatasmithScene = NewObject<UDatasmithScene>(Arguments.Parent, DatasmithSceneClass, *Arguments.AssetName, RF_Public | RF_Standalone);
+	}
+	else if (ExistingAsset->GetClass()->IsChildOf(DatasmithSceneClass))
+	{
+		//This is a reimport, we are just re-updating the source data
+		DatasmithScene = static_cast<UDatasmithScene*>(ExistingAsset);
+	}
+
+	if (!DatasmithScene)
+	{
+		UE_LOG(LogInterchangeDatasmith, Warning, TEXT("Could not create Datasmith Scene asset %s"), *Arguments.AssetName);
+		return nullptr;
+	}
+#endif //WITH_EDITORONLY_DATA
+
+	return DatasmithScene;
+}
+
+UObject* UInterchangeDatasmithSceneFactory::CreateAsset(const FCreateAssetParams& Arguments)
+{
+ #if !WITH_EDITORONLY_DATA
+ 	UE_LOG(LogInterchangeDatasmith, Error, TEXT("Cannot import datasmith scene asset in runtime, this is an editor only feature."));
+ 	return nullptr; 
+ #else //WITH_EDITORONLY_DATA
+
+	TRACE_CPUPROFILER_EVENT_SCOPE(UInterchangeTextureFactory::CreateAsset);
+
+	if (!Arguments.AssetNode)
+	{
+		UE_LOG(LogInterchangeDatasmith, Error, TEXT("UInterchangeDatasmithSceneFactory: Asset node parameter is null."));
+		return nullptr;
+	}
+
+	const UClass* DatasmithSceneClass = Arguments.AssetNode->GetObjectClass();
+	if (!DatasmithSceneClass || !DatasmithSceneClass->IsChildOf(UDatasmithScene::StaticClass()))
+	{
+		UE_LOG(LogInterchangeDatasmith, Error, TEXT("UInterchangeDatasmithSceneFactory: Asset node parameter class doesnt derive from UDatasmithScene."));
+		return nullptr;
+	}
+
+	UClass* SupportedFactoryNodeClass = Arguments.AssetNode->GetClass() == UInterchangeDatasmithSceneFactoryNode::StaticClass()
+		? Arguments.AssetNode->GetClass()
+		: nullptr;
+	if (SupportedFactoryNodeClass == nullptr)
+	{
+		UE_LOG(LogInterchangeDatasmith, Error, TEXT("UInterchangeDatasmithSceneFactory: Asset node parameter is not a UInterchangeDatasmithSceneFactoryNode."));
+		return nullptr;
+	}
+
+	// create an asset if it doesn't exist
+	UObject* ExistingAsset = StaticFindObject(nullptr, Arguments.Parent, *Arguments.AssetName);
+	UDatasmithScene* DatasmithScene = nullptr;
+
+	// create a new texture or overwrite existing asset, if possible
+	if (!ExistingAsset)
+	{
+		//NewObject is not thread safe, the asset registry directory watcher tick on the main thread can trig before we finish initializing the UObject and will crash
+		//The UObject should have been create by calling CreateEmptyAsset on the main thread.
+		check(IsInGameThread());
+		DatasmithScene = NewObject<UDatasmithScene>(Arguments.Parent, DatasmithSceneClass, *Arguments.AssetName, RF_Public | RF_Standalone);
+	}
+	else if (ExistingAsset->GetClass()->IsChildOf(DatasmithSceneClass))
+	{
+		//This is a reimport, we are just re-updating the source data
+		DatasmithScene = static_cast<UDatasmithScene*>(ExistingAsset);
+	}
+
+	if (!DatasmithScene)
+	{
+		UE_LOG(LogInterchangeDatasmith, Error, TEXT("UInterchangeDatasmithSceneFactory: Could not create datasmith scene asset %s"), *Arguments.AssetName);
+		return nullptr;
+	}
+
+	// TODO add InterchangeAssetImportData or DatasmithAssetImportData.
+	// TODO link the created asset to the factory.
+
+	return DatasmithScene;
+#endif //WITH_EDITORONLY_DATA
+}
+
+void UInterchangeDatasmithSceneFactory::PreImportPreCompletedCallback(const FImportPreCompletedCallbackParams& Arguments)
+{
+
+}
\ No newline at end of file
diff --git a/Engine/Plugins/Experimental/Enterprise/DatasmithInterchange/Source/DatasmithInterchange/Private/InterchangeDatasmithSceneFactoryNode.cpp b/Engine/Plugins/Experimental/Enterprise/DatasmithInterchange/Source/DatasmithInterchange/Private/InterchangeDatasmithSceneFactoryNode.cpp
new file mode 100644
--- /dev/null
+++ b/Engine/Plugins/Experimental/Enterprise/DatasmithInterchange/Source/DatasmithInterchange/Private/InterchangeDatasmithSceneFactoryNode.cpp
@@ -0,0 +1,16 @@
+// Copyright Epic Games, Inc. All Rights Reserved.
+
+#include "InterchangeDatasmithSceneFactoryNode.h"
+
+#include "DatasmithScene.h"
+
+FString UInterchangeDatasmithSceneFactoryNode::GetTypeName() const
+{
+	const FString TypeName = TEXT("DatasmithSceneFactoryNode");
+	return TypeName;
+}
+
+UClass* UInterchangeDatasmithSceneFactoryNode::GetObjectClass() const
+{
+	return UDatasmithScene::StaticClass();
+}
\ No newline at end of file
diff --git a/Engine/Plugins/Experimental/Enterprise/DatasmithInterchange/Source/DatasmithInterchange/Private/InterchangeDatasmithSceneNode.cpp b/Engine/Plugins/Experimental/Enterprise/DatasmithInterchange/Source/DatasmithInterchange/Private/InterchangeDatasmithSceneNode.cpp
new file mode 100644
--- /dev/null
+++ b/Engine/Plugins/Experimental/Enterprise/DatasmithInterchange/Source/DatasmithInterchange/Private/InterchangeDatasmithSceneNode.cpp
@@ -0,0 +1,3 @@
+// Copyright Epic Games, Inc. All Rights Reserved.
+
+#include "InterchangeDatasmithSceneNode.h"
\ No newline at end of file
diff --git a/Engine/Plugins/Experimental/Enterprise/DatasmithInterchange/Source/DatasmithInterchange/Private/InterchangeDatasmithStaticMeshPipeline.cpp b/Engine/Plugins/Experimental/Enterprise/DatasmithInterchange/Source/DatasmithInterchange/Private/InterchangeDatasmithStaticMeshPipeline.cpp
new file mode 100644
--- /dev/null
+++ b/Engine/Plugins/Experimental/Enterprise/DatasmithInterchange/Source/DatasmithInterchange/Private/InterchangeDatasmithStaticMeshPipeline.cpp
@@ -0,0 +1,94 @@
+// Copyright Epic Games, Inc. All Rights Reserved. 
+
+#include "InterchangeDatasmithStaticMeshPipeline.h"
+
+#include "InterchangeDatasmithUtils.h"
+
+#include "InterchangeMaterialFactoryNode.h"
+#include "InterchangeStaticMeshFactoryNode.h"
+#include "InterchangeMeshNode.h"
+
+#include "Engine/StaticMesh.h"
+#include "Materials/MaterialInterface.h"
+
+void UInterchangeDatasmithStaticMeshPipeline::ExecutePreImportPipeline(UInterchangeBaseNodeContainer* NodeContainer, const TArray<UInterchangeSourceData*>& InSourceDatas)
+{
+	using namespace UE::DatasmithInterchange;
+
+	Super::ExecutePreImportPipeline(NodeContainer, InSourceDatas);
+
+	// Add material factory dependencies for meshes where all slots are filled with the same material
+	for (UInterchangeStaticMeshFactoryNode* MeshFactoryNode : NodeUtils::GetNodes<UInterchangeStaticMeshFactoryNode>(NodeContainer))
+	{
+		TArray<FString> TargetNodes;
+		MeshFactoryNode->GetTargetNodeUids(TargetNodes);
+		if (TargetNodes.Num() == 0)
+		{
+			continue;
+		}
+
+		const UInterchangeMeshNode* MeshNode = Cast< UInterchangeMeshNode>(NodeContainer->GetNode(TargetNodes[0]));
+		if (!MeshNode)
+		{
+			continue;
+		}
+
+		FString MaterialUid;
+		if (!MeshNode->GetStringAttribute(MeshUtils::MeshMaterialAttrName, MaterialUid))
+		{
+			continue;
+		}
+
+		const FString MaterialFactoryUid = UInterchangeMaterialFactoryNode::GetMaterialFactoryNodeUidFromMaterialNodeUid(MaterialUid);
+		MeshFactoryNode->AddFactoryDependencyUid(MaterialFactoryUid);
+		MeshFactoryNode->AddStringAttribute(MeshUtils::MeshMaterialAttrName, MaterialFactoryUid);
+	}
+}
+
+void UInterchangeDatasmithStaticMeshPipeline::ExecutePostImportPipeline(const UInterchangeBaseNodeContainer* NodeContainer, const FString& FactoryNodeKey, UObject* CreatedAsset, bool bIsAReimport)
+{
+	using namespace UE::DatasmithInterchange;
+
+	if (!NodeContainer || !CreatedAsset)
+	{
+		return;
+	}
+
+	Super::ExecutePostImportPipeline(NodeContainer, FactoryNodeKey, CreatedAsset, bIsAReimport);
+
+	// If applicable, update FStaticMaterial of newly create mesh
+	UStaticMesh* StaticMesh = Cast<UStaticMesh>(CreatedAsset);
+	if (!StaticMesh)
+	{
+		return;
+	}
+
+	const UInterchangeStaticMeshFactoryNode* FactoryNode = Cast< UInterchangeStaticMeshFactoryNode>(NodeContainer->GetFactoryNode(FactoryNodeKey));
+	if (!FactoryNode)
+	{
+		return;
+	}
+
+	FString MaterialFactoryUid;
+	if (!FactoryNode->GetStringAttribute(MeshUtils::MeshMaterialAttrName, MaterialFactoryUid))
+	{
+		return;
+	}
+
+	const UInterchangeBaseMaterialFactoryNode* MaterialFactoryNode = Cast< UInterchangeBaseMaterialFactoryNode>(NodeContainer->GetFactoryNode(MaterialFactoryUid));
+	if (!MaterialFactoryNode)
+	{
+		return;
+	}
+
+	FSoftObjectPath MaterialFactoryNodeReferenceObject;
+	MaterialFactoryNode->GetCustomReferenceObject(MaterialFactoryNodeReferenceObject);
+	if (UMaterialInterface* MaterialInterface = Cast<UMaterialInterface>(MaterialFactoryNodeReferenceObject.ResolveObject()))
+	{
+		TArray<FStaticMaterial>& StaticMaterials = StaticMesh->GetStaticMaterials();
+		for (FStaticMaterial& StaticMaterial : StaticMaterials)
+		{
+			StaticMaterial.MaterialInterface = MaterialInterface;
+		}
+	}
+}
diff --git a/Engine/Plugins/Experimental/Enterprise/DatasmithInterchange/Source/DatasmithInterchange/Private/InterchangeDatasmithStaticMeshPipeline.h b/Engine/Plugins/Experimental/Enterprise/DatasmithInterchange/Source/DatasmithInterchange/Private/InterchangeDatasmithStaticMeshPipeline.h
new file mode 100644
--- /dev/null
+++ b/Engine/Plugins/Experimental/Enterprise/DatasmithInterchange/Source/DatasmithInterchange/Private/InterchangeDatasmithStaticMeshPipeline.h
@@ -0,0 +1,28 @@
+// Copyright Epic Games, Inc. All Rights Reserved.
+
+#pragma once
+
+#include "CoreMinimal.h"
+
+#include "InterchangeGenericMeshPipeline.h"
+
+#include "InterchangeDatasmithStaticMeshPipeline.generated.h"
+
+class IDatasmithScene;
+class UMaterialInterface;
+
+UCLASS(BlueprintType, Experimental)
+class DATASMITHINTERCHANGE_API UInterchangeDatasmithStaticMeshPipeline : public UInterchangeGenericMeshPipeline
+{
+	GENERATED_BODY()
+
+protected:
+	virtual void ExecutePreImportPipeline(UInterchangeBaseNodeContainer* BaseNodeContainer, const TArray<UInterchangeSourceData*>& SourceDatas) override;
+	virtual void ExecutePostImportPipeline(const UInterchangeBaseNodeContainer* BaseNodeContainer, const FString& NodeKey, UObject* CreatedAsset, bool bIsAReimport) override;
+
+	virtual bool CanExecuteOnAnyThread(EInterchangePipelineTask PipelineTask) override
+	{
+		// This pipeline creates UObjects and assets. Not safe to execute outside of main thread.
+		return false;
+	}
+};
\ No newline at end of file
diff --git a/Engine/Plugins/Experimental/Enterprise/DatasmithInterchange/Source/DatasmithInterchange/Private/InterchangeDatasmithTextureData.cpp b/Engine/Plugins/Experimental/Enterprise/DatasmithInterchange/Source/DatasmithInterchange/Private/InterchangeDatasmithTextureData.cpp
new file mode 100644
--- /dev/null
+++ b/Engine/Plugins/Experimental/Enterprise/DatasmithInterchange/Source/DatasmithInterchange/Private/InterchangeDatasmithTextureData.cpp
@@ -0,0 +1,180 @@
+// Copyright Epic Games, Inc. All Rights Reserved.
+
+#include "InterchangeDatasmithTextureData.h"
+
+#include "InterchangeTextureNode.h"
+#include "IDatasmithSceneElements.h"
+
+#define IMPLEMENT_GET_ATTRIBUTE(Key) ConstBaseNode ? ConstBaseNode->GetAttribute(Key, AttributeValue) : false
+
+#define IMPLEMENT_SET_ATTRIBUTE(Key) BaseNode ? BaseNode->SetAttribute(Key, AttributeValue) : false
+
+namespace UE::DatasmithInterchange
+{
+	const FName FInterchangeDatasmithTextureDataConst::DatasmithTextureDataKey = FName(TEXT("DatasmithTextureData"));
+	const FName FInterchangeDatasmithTextureDataConst::FileKey = FName(TEXT("File"));
+	const FName FInterchangeDatasmithTextureDataConst::TextureModeKey = FName(TEXT("TextureMode"));
+	const FName FInterchangeDatasmithTextureDataConst::TextureFilterKey = FName(TEXT("TextureFilter"));
+	const FName FInterchangeDatasmithTextureDataConst::TextureAddressXKey = FName(TEXT("TextureAddressX"));
+	const FName FInterchangeDatasmithTextureDataConst::TextureAddressYKey = FName(TEXT("TextureAddressY"));
+	const FName FInterchangeDatasmithTextureDataConst::AllowResizeKey = FName(TEXT("AllowResize"));
+	const FName FInterchangeDatasmithTextureDataConst::RGBCurveKey = FName(TEXT("RGBCurve"));
+	const FName FInterchangeDatasmithTextureDataConst::SRGBKey = FName(TEXT("SRGB"));
+
+	FInterchangeDatasmithTextureDataConst::FInterchangeDatasmithTextureDataConst(const UInterchangeBaseNode* InNode)
+		: ConstBaseNode(InNode)
+	{}
+
+	FInterchangeDatasmithTextureData::FInterchangeDatasmithTextureData(UInterchangeBaseNode* InNode)
+		: FInterchangeDatasmithTextureDataConst(InNode)
+		, BaseNode(InNode)
+	{
+		BaseNode->AddBooleanAttribute(DatasmithTextureDataKey, true);
+	}
+
+	bool FInterchangeDatasmithTextureDataConst::HasData(const UInterchangeBaseNode* Node)
+	{
+		bool bTextureData;
+		if (Node && Node->GetAttribute(DatasmithTextureDataKey, bTextureData) && bTextureData)
+		{
+			return true;
+		}
+
+		return false;
+	}
+
+	bool FInterchangeDatasmithTextureDataConst::GetCustomFile(FString& AttributeValue) const
+	{
+		return IMPLEMENT_GET_ATTRIBUTE(FileKey);
+	}
+
+	bool FInterchangeDatasmithTextureData::SetCustomFile(const FString& AttributeValue)
+	{
+		return IMPLEMENT_SET_ATTRIBUTE(FileKey);
+	}
+
+	bool FInterchangeDatasmithTextureDataConst::GetCustomTextureMode(EDatasmithTextureMode& InAttributeValue) const
+	{
+		uint8 AttributeValue = 0;
+		if (IMPLEMENT_GET_ATTRIBUTE(TextureModeKey))
+		{
+			InAttributeValue = static_cast<EDatasmithTextureMode>(AttributeValue);
+			return true;
+		}
+		return false;
+	}
+
+	bool FInterchangeDatasmithTextureData::SetCustomTextureMode(const EDatasmithTextureMode& InAttributeValue)
+	{
+		uint8 AttributeValue = static_cast<uint8>(InAttributeValue);
+		return IMPLEMENT_SET_ATTRIBUTE(TextureModeKey);
+	}
+
+	bool FInterchangeDatasmithTextureDataConst::GetCustomTextureFilter(EDatasmithTextureFilter& InAttributeValue) const
+	{
+		uint8 AttributeValue = 0;
+		if (IMPLEMENT_GET_ATTRIBUTE(TextureFilterKey))
+		{
+			InAttributeValue = static_cast<EDatasmithTextureFilter>(AttributeValue);
+			return true;
+		}
+		return false;
+	}
+
+	bool FInterchangeDatasmithTextureData::SetCustomTextureFilter(const EDatasmithTextureFilter& InAttributeValue)
+	{
+		uint8 AttributeValue = static_cast<uint8>(InAttributeValue);
+		return IMPLEMENT_SET_ATTRIBUTE(TextureFilterKey);
+	}
+
+	bool FInterchangeDatasmithTextureDataConst::GetCustomTextureAddressX(EDatasmithTextureAddress& InAttributeValue) const
+	{
+		uint8 AttributeValue = 0;
+		if (IMPLEMENT_GET_ATTRIBUTE(TextureAddressXKey))
+		{
+			InAttributeValue = static_cast<EDatasmithTextureAddress>(AttributeValue);
+			return true;
+		}
+		return false;
+	}
+
+	bool FInterchangeDatasmithTextureData::SetCustomTextureAddressX(const EDatasmithTextureAddress& InAttributeValue)
+	{
+		uint8 AttributeValue = static_cast<uint8>(InAttributeValue);
+		return IMPLEMENT_SET_ATTRIBUTE(TextureAddressXKey);
+	}
+
+	bool FInterchangeDatasmithTextureDataConst::GetCustomTextureAddressY(EDatasmithTextureAddress& InAttributeValue) const
+	{
+		uint8 AttributeValue = 0;
+		if (IMPLEMENT_GET_ATTRIBUTE(TextureAddressYKey))
+		{
+			InAttributeValue = static_cast<EDatasmithTextureAddress>(AttributeValue);
+			return true;
+		}
+		return false;
+	}
+
+	bool FInterchangeDatasmithTextureData::SetCustomTextureAddressY(const EDatasmithTextureAddress& InAttributeValue)
+	{
+		uint8 AttributeValue = static_cast<uint8>(InAttributeValue);
+		return IMPLEMENT_SET_ATTRIBUTE(TextureAddressYKey);
+	}
+
+	bool FInterchangeDatasmithTextureDataConst::GetCustomAllowResize(bool& AttributeValue) const
+	{
+		return IMPLEMENT_GET_ATTRIBUTE(AllowResizeKey);
+	}
+
+	bool FInterchangeDatasmithTextureData::SetCustomAllowResize(const bool& AttributeValue)
+	{
+		return IMPLEMENT_SET_ATTRIBUTE(AllowResizeKey);
+	}
+
+	bool FInterchangeDatasmithTextureDataConst::GetCustomRGBCurve(float& AttributeValue) const
+	{
+		return IMPLEMENT_GET_ATTRIBUTE(RGBCurveKey);
+	}
+
+	bool FInterchangeDatasmithTextureData::SetCustomRGBCurve(const float& AttributeValue)
+	{
+		return IMPLEMENT_SET_ATTRIBUTE(RGBCurveKey);
+	}
+
+	bool FInterchangeDatasmithTextureDataConst::GetCustomSRGB(EDatasmithColorSpace& InAttributeValue) const
+	{
+		uint8 AttributeValue = 0;
+		if (IMPLEMENT_GET_ATTRIBUTE(SRGBKey))
+		{
+			InAttributeValue = static_cast<EDatasmithColorSpace>(AttributeValue);
+			return true;
+		}
+		return false;
+	}
+
+	bool FInterchangeDatasmithTextureData::SetCustomSRGB(const EDatasmithColorSpace& InAttributeValue)
+	{
+		uint8 AttributeValue = static_cast<uint8>(InAttributeValue);
+		return IMPLEMENT_SET_ATTRIBUTE(SRGBKey);
+	}
+
+	namespace TextureUtils
+	{
+		void ApplyTextureElementToNode(const TSharedRef<IDatasmithTextureElement>& TextureElement, UInterchangeBaseNode* BaseNode)
+		{
+			FInterchangeDatasmithTextureData TextureData(BaseNode);
+
+			TextureData.SetCustomFile(TextureElement->GetFile());
+			TextureData.SetCustomTextureMode(TextureElement->GetTextureMode());
+			TextureData.SetCustomTextureFilter(TextureElement->GetTextureFilter());
+			TextureData.SetCustomTextureAddressX(TextureElement->GetTextureAddressX());
+			TextureData.SetCustomTextureAddressY(TextureElement->GetTextureAddressY());
+			TextureData.SetCustomAllowResize(TextureElement->GetAllowResize());
+			TextureData.SetCustomRGBCurve(TextureElement->GetRGBCurve());
+			TextureData.SetCustomSRGB(TextureElement->GetSRGB());
+		}
+	}
+}
+
+#undef IMPLEMENT_SET_ATTRIBUTE
+#undef IMPLEMENT_GET_ATTRIBUTE
\ No newline at end of file
diff --git a/Engine/Plugins/Experimental/Enterprise/DatasmithInterchange/Source/DatasmithInterchange/Private/InterchangeDatasmithTextureData.h b/Engine/Plugins/Experimental/Enterprise/DatasmithInterchange/Source/DatasmithInterchange/Private/InterchangeDatasmithTextureData.h
new file mode 100644
--- /dev/null
+++ b/Engine/Plugins/Experimental/Enterprise/DatasmithInterchange/Source/DatasmithInterchange/Private/InterchangeDatasmithTextureData.h
@@ -0,0 +1,84 @@
+// Copyright Epic Games, Inc. All Rights Reserved.
+
+#pragma once
+
+#include "CoreMinimal.h"
+#include "DatasmithDefinitions.h"
+#include "Nodes/InterchangeBaseNode.h"
+
+class IDatasmithTextureElement;
+class UInterchangeTextureNode;
+
+namespace UE::DatasmithInterchange
+{
+	/**
+	 * Helper struct to get/set Datasmith texture attributes on a node.
+	 */
+	struct FInterchangeDatasmithTextureDataConst
+	{
+		static bool HasData(const UInterchangeBaseNode* Node);
+
+		FInterchangeDatasmithTextureDataConst(const UInterchangeBaseNode* InNode);
+
+		bool GetCustomFile(FString& AttributeValue) const;
+
+		bool GetCustomTextureMode(EDatasmithTextureMode& AttributeValue) const;
+
+		bool GetCustomTextureFilter(EDatasmithTextureFilter& AttributeValue) const;
+
+		bool GetCustomTextureAddressX(EDatasmithTextureAddress& AttributeValue) const;
+
+		bool GetCustomTextureAddressY(EDatasmithTextureAddress& AttributeValue) const;
+
+		bool GetCustomAllowResize(bool& AttributeValue) const;
+
+		bool GetCustomRGBCurve(float& AttributeValue) const;
+
+		bool GetCustomSRGB(EDatasmithColorSpace& AttributeValue) const;
+
+	protected:
+		const UInterchangeBaseNode* ConstBaseNode;
+
+		static const FName DatasmithTextureDataKey;
+		static const FName FileKey;
+		static const FName TextureModeKey;
+		static const FName TextureFilterKey;
+		static const FName TextureAddressXKey;
+		static const FName TextureAddressYKey;
+		static const FName AllowResizeKey;
+		static const FName RGBCurveKey;
+		static const FName SRGBKey;
+	};
+
+	struct FInterchangeDatasmithTextureData : public FInterchangeDatasmithTextureDataConst
+	{
+		FInterchangeDatasmithTextureData(UInterchangeBaseNode* InNode);
+
+		bool SetCustomFile(const FString& AttributeValue);
+
+		bool SetCustomTextureMode(const EDatasmithTextureMode& AttributeValue);
+
+		bool SetCustomTextureFilter(const EDatasmithTextureFilter& AttributeValue);
+
+		bool SetCustomTextureAddressX(const EDatasmithTextureAddress& AttributeValue);
+
+		bool SetCustomTextureAddressY(const EDatasmithTextureAddress& AttributeValue);
+
+		bool SetCustomAllowResize(const bool& AttributeValue);
+
+		bool SetCustomRGBCurve(const float& AttributeValue);
+
+		bool SetCustomSRGB(const EDatasmithColorSpace& AttributeValue);
+
+	private:
+		UInterchangeBaseNode* BaseNode;
+	};
+
+	namespace TextureUtils
+	{
+		/**
+		 * Register the properties of the TextureElement into the BaseNode's attributes.
+		 */
+		void ApplyTextureElementToNode(const TSharedRef<IDatasmithTextureElement>& TextureElement, UInterchangeBaseNode* BaseNode);
+	}
+}
\ No newline at end of file
diff --git a/Engine/Plugins/Experimental/Enterprise/DatasmithInterchange/Source/DatasmithInterchange/Private/InterchangeDatasmithTexturePipeline.cpp b/Engine/Plugins/Experimental/Enterprise/DatasmithInterchange/Source/DatasmithInterchange/Private/InterchangeDatasmithTexturePipeline.cpp
new file mode 100644
--- /dev/null
+++ b/Engine/Plugins/Experimental/Enterprise/DatasmithInterchange/Source/DatasmithInterchange/Private/InterchangeDatasmithTexturePipeline.cpp
@@ -0,0 +1,167 @@
+// Copyright Epic Games, Inc. All Rights Reserved. 
+
+#include "InterchangeDatasmithTexturePipeline.h"
+
+#include "InterchangeDatasmithTextureData.h"
+#include "InterchangeDatasmithUtils.h"
+
+#include "InterchangeDatasmithSceneNode.h"
+
+#include "InterchangeTexture2DFactoryNode.h"
+#include "InterchangeTextureNode.h"
+#include "IDatasmithSceneElements.h"
+
+void UInterchangeDatasmithTexturePipeline::ExecutePreImportPipeline(UInterchangeBaseNodeContainer* InNodeContainer, const TArray<UInterchangeSourceData*>& InSourceDatas)
+{
+	using namespace UE::DatasmithInterchange;
+
+	Super::ExecutePreImportPipeline(InNodeContainer, InSourceDatas);
+
+	for (UInterchangeTextureFactoryNode* TextureFactoryNode : NodeUtils::GetNodes<UInterchangeTextureFactoryNode>(InNodeContainer))
+	{
+		PreImportTextureFactoryNode(InNodeContainer, TextureFactoryNode);
+	}
+}
+
+void UInterchangeDatasmithTexturePipeline::PreImportTextureFactoryNode(UInterchangeBaseNodeContainer* InNodeContainer, UInterchangeTextureFactoryNode* TextureFactoryNode) const
+{
+	using namespace UE::DatasmithInterchange;
+	
+	TArray<FString> TargetNodes;
+	TextureFactoryNode->GetTargetNodeUids(TargetNodes);
+	if (TargetNodes.Num() == 0)
+	{
+		return;
+	}
+
+	const UInterchangeBaseNode* TargetNode = InNodeContainer->GetNode(TargetNodes[0]);
+	if (!FInterchangeDatasmithTextureData::HasData(TargetNode))
+	{
+		return;
+	}
+
+	FInterchangeDatasmithTextureDataConst TextureData(TargetNode);
+	TOptional< bool > bLocalFlipNormalMapGreenChannel;
+	TOptional< TextureMipGenSettings > MipGenSettings;
+	TOptional< TextureGroup > LODGroup;
+	TOptional< TextureCompressionSettings > CompressionSettings;
+
+	// Make sure to set the proper LODGroup as it's used to determine the CompressionSettings when using TEXTUREGROUP_WorldNormalMap
+	EDatasmithTextureMode TextureMode;
+	if (TextureData.GetCustomTextureMode(TextureMode))
+	{
+		switch (TextureMode)
+		{
+		case EDatasmithTextureMode::Diffuse:
+			LODGroup = TEXTUREGROUP_World;
+			break;
+		case EDatasmithTextureMode::Specular:
+			LODGroup = TEXTUREGROUP_WorldSpecular;
+			break;
+		case EDatasmithTextureMode::Bump:
+		case EDatasmithTextureMode::Normal:
+			LODGroup = TEXTUREGROUP_WorldNormalMap;
+			CompressionSettings = TC_Normalmap;
+			break;
+		case EDatasmithTextureMode::NormalGreenInv:
+			LODGroup = TEXTUREGROUP_WorldNormalMap;
+			CompressionSettings = TC_Normalmap;
+			bLocalFlipNormalMapGreenChannel = true;
+			break;
+		}
+	}
+
+	const TOptional< float > RGBCurve = [&TextureData]() -> TOptional< float >
+	{
+		float ElementRGBCurve;
+
+		if (TextureData.GetCustomRGBCurve(ElementRGBCurve)
+			&& FMath::IsNearlyEqual(ElementRGBCurve, 1.0f) == false
+			&& ElementRGBCurve > 0.f)
+		{
+			return ElementRGBCurve;
+		}
+
+		return {};
+	}();
+
+	static_assert(TextureAddress::TA_Wrap == (int)EDatasmithTextureAddress::Wrap && TextureAddress::TA_Mirror == (int)EDatasmithTextureAddress::Mirror, "Texture Address enum doesn't match!");
+
+	TOptional< TextureFilter > TexFilter;
+	EDatasmithTextureFilter TextureFilter;
+	if (TextureData.GetCustomTextureFilter(TextureFilter))
+	{
+		switch (TextureFilter)
+		{
+		case EDatasmithTextureFilter::Nearest:
+			TexFilter = TextureFilter::TF_Nearest;
+			break;
+		case EDatasmithTextureFilter::Bilinear:
+			TexFilter = TextureFilter::TF_Bilinear;
+			break;
+		case EDatasmithTextureFilter::Trilinear:
+			TexFilter = TextureFilter::TF_Trilinear;
+			break;
+		}
+	}
+
+	TOptional< bool > bSrgb;
+	EDatasmithColorSpace ColorSpace;
+	if (TextureData.GetCustomSRGB(ColorSpace))
+	{
+		if (ColorSpace == EDatasmithColorSpace::sRGB)
+		{
+			bSrgb = true;
+		}
+		else if (ColorSpace == EDatasmithColorSpace::Linear)
+		{
+			bSrgb = false;
+		}
+	}
+
+	EDatasmithTextureAddress AddressX;
+	EDatasmithTextureAddress AddressY;
+	UInterchangeTexture2DFactoryNode* Texture2DFactoryNode = Cast<UInterchangeTexture2DFactoryNode>(TextureFactoryNode);
+	if (Texture2DFactoryNode
+		&& TextureData.GetCustomTextureAddressX(AddressX)
+		&& TextureData.GetCustomTextureAddressY(AddressY))
+	{
+		Texture2DFactoryNode->SetCustomAddressX((TextureAddress)AddressX);
+		Texture2DFactoryNode->SetCustomAddressY((TextureAddress)AddressY);
+	}
+
+	if (bSrgb.IsSet())
+	{
+		TextureFactoryNode->SetCustomSRGB(bSrgb.GetValue());
+	}
+
+	if (bLocalFlipNormalMapGreenChannel.IsSet())
+	{
+		TextureFactoryNode->SetCustombFlipGreenChannel(bLocalFlipNormalMapGreenChannel.GetValue());
+	}
+
+	if (MipGenSettings.IsSet())
+	{
+		TextureFactoryNode->SetCustomMipGenSettings(MipGenSettings.GetValue());
+	}
+
+	if (LODGroup.IsSet())
+	{
+		TextureFactoryNode->SetCustomLODGroup(LODGroup.GetValue());
+	}
+
+	if (CompressionSettings.IsSet())
+	{
+		TextureFactoryNode->SetCustomCompressionSettings(CompressionSettings.GetValue());
+	}
+
+	if (RGBCurve.IsSet())
+	{
+		TextureFactoryNode->SetCustomAdjustRGBCurve(RGBCurve.GetValue());
+	}
+
+	if (TexFilter.IsSet())
+	{
+		TextureFactoryNode->SetCustomFilter(TexFilter.GetValue());
+	}
+}
\ No newline at end of file
diff --git a/Engine/Plugins/Experimental/Enterprise/DatasmithInterchange/Source/DatasmithInterchange/Private/InterchangeDatasmithTexturePipeline.h b/Engine/Plugins/Experimental/Enterprise/DatasmithInterchange/Source/DatasmithInterchange/Private/InterchangeDatasmithTexturePipeline.h
new file mode 100644
--- /dev/null
+++ b/Engine/Plugins/Experimental/Enterprise/DatasmithInterchange/Source/DatasmithInterchange/Private/InterchangeDatasmithTexturePipeline.h
@@ -0,0 +1,27 @@
+// Copyright Epic Games, Inc. All Rights Reserved.
+
+#pragma once
+
+#include "CoreMinimal.h"
+
+#include "InterchangeGenericTexturePipeline.h"
+
+#include "InterchangeDatasmithTexturePipeline.generated.h"
+
+class IDatasmithScene;
+class UInterchangeBaseNodeContainer;
+class UInterchangeSourceData;
+class UInterchangeTextureFactoryNode;
+class UObject;
+
+UCLASS(BlueprintType, Experimental)
+class DATASMITHINTERCHANGE_API UInterchangeDatasmithTexturePipeline : public UInterchangeGenericTexturePipeline
+{
+	GENERATED_BODY()
+
+public:
+	virtual void ExecutePreImportPipeline(UInterchangeBaseNodeContainer* BaseNodeContainer, const TArray<UInterchangeSourceData*>& SourceDatas) override;
+
+private:
+	void PreImportTextureFactoryNode(UInterchangeBaseNodeContainer* BaseNodeContainer, UInterchangeTextureFactoryNode* TextureFactoryNode) const;
+};
\ No newline at end of file
diff --git a/Engine/Plugins/Experimental/Enterprise/DatasmithInterchange/Source/DatasmithInterchange/Private/InterchangeDatasmithTranslator.cpp b/Engine/Plugins/Experimental/Enterprise/DatasmithInterchange/Source/DatasmithInterchange/Private/InterchangeDatasmithTranslator.cpp
new file mode 100644
--- /dev/null
+++ b/Engine/Plugins/Experimental/Enterprise/DatasmithInterchange/Source/DatasmithInterchange/Private/InterchangeDatasmithTranslator.cpp
@@ -0,0 +1,649 @@
+// Copyright Epic Games, Inc. All Rights Reserved.
+
+#include "InterchangeDatasmithTranslator.h"
+
+#include "InterchangeDatasmithAreaLightNode.h"
+#include "InterchangeDatasmithLog.h"
+#include "InterchangeDatasmithMaterialNode.h"
+#include "InterchangeDatasmithSceneNode.h"
+#include "InterchangeDatasmithTextureData.h"
+#include "InterchangeDatasmithUtils.h"
+
+#include "DatasmithAnimationElements.h"
+#include "DatasmithMaterialElements.h"
+#include "DatasmithScene.h"
+#include "DatasmithSceneSource.h"
+#include "DatasmithTranslatableSource.h"
+#include "DatasmithTranslatorManager.h"
+#include "DatasmithUtils.h"
+#include "DatasmithVariantElements.h"
+#include "IDatasmithSceneElements.h"
+
+#include "ExternalSourceModule.h"
+#include "SourceUri.h"
+#include "InterchangeCameraNode.h"
+#include "InterchangeAnimationTrackSetNode.h"
+#include "InterchangeLightNode.h"
+#include "InterchangeManager.h"
+#include "InterchangeMaterialDefinitions.h"
+#include "InterchangeMeshNode.h"
+#include "InterchangeShaderGraphNode.h"
+#include "InterchangeSceneNode.h"
+#include "InterchangeTexture2DNode.h"
+#include "InterchangeVariantSetNode.h"
+
+#define LOCTEXT_NAMESPACE "DatasmithInterchange"
+
+bool UInterchangeDatasmithTranslator::CanImportSourceData(const UInterchangeSourceData* InSourceData) const
+{
+	using namespace UE::DatasmithImporter;
+
+	const FString FilePath = InSourceData->GetFilename();
+	const FString FileExtension = FPaths::GetExtension(FilePath);
+	if (FileExtension == TEXT("gltf") || FileExtension == TEXT("glb"))
+	{
+		// Do not translate gltf since there is already a native gltf interchange translator. 
+		return false;
+	}
+
+	const FSourceUri FileNameUri = FSourceUri::FromFilePath(FilePath);
+	TSharedPtr<FExternalSource> ExternalSource = IExternalSourceModule::GetOrCreateExternalSource(FileNameUri);
+
+	return ExternalSource.IsValid() && ExternalSource->IsAvailable();
+}
+
+bool UInterchangeDatasmithTranslator::Translate(UInterchangeBaseNodeContainer& BaseNodeContainer) const
+{
+	using namespace UE::DatasmithImporter;
+	using namespace UE::DatasmithInterchange;
+
+	// TODO: This code should eventually go into UInterchangeTranslatorBase once the FExternalSource module gets integrated into Interchange
+	FString FilePath = SourceData->GetFilename();
+	FileName = FPaths::GetCleanFilename(FilePath);
+	const FSourceUri FileNameUri = FSourceUri::FromFilePath(FilePath);
+	const_cast<UInterchangeDatasmithTranslator*>(this)->LoadedExternalSource = IExternalSourceModule::GetOrCreateExternalSource(FileNameUri);
+
+	if (!LoadedExternalSource.IsValid() || !LoadedExternalSource->IsAvailable())
+	{
+		return false;
+	}
+
+	UClass* Class = UInterchangeDatasmithSceneNode::StaticClass();
+	if (!ensure(Class))
+	{
+		return false;
+	}
+
+	StartTime = FPlatformTime::Cycles64();
+	FPaths::NormalizeFilename(FilePath);
+
+	// Should it be mutable instead? If Translate is const should we really be doing this?.
+	TSharedPtr<IDatasmithScene> DatasmithScene = LoadedExternalSource->TryLoad();
+	if (!DatasmithScene.IsValid())
+	{
+		return false;
+	}
+
+	// Datasmith Scene Node
+	{
+		FString DisplayLabel = DatasmithScene->GetName();
+		FString NodeUID(NodeUtils::DatasmithScenePrefix + FilePath);
+		UInterchangeDatasmithSceneNode* DatasmithSceneNode = NewObject<UInterchangeDatasmithSceneNode>(&BaseNodeContainer, Class);
+		if (!ensure(DatasmithSceneNode))
+		{
+			return false;
+		}
+		DatasmithSceneNode->InitializeNode(NodeUID, DisplayLabel, EInterchangeNodeContainerType::TranslatedAsset);
+
+		// Assigning those variables are our only way to pass the translated state to the part of the pipeline that does not use interchange yet.
+		DatasmithSceneNode->ExternalSource = LoadedExternalSource;
+		DatasmithSceneNode->DatasmithScene = DatasmithScene;
+
+		BaseNodeContainer.AddNode(DatasmithSceneNode);
+	}
+
+	// Texture Nodes
+	{
+		FDatasmithUniqueNameProvider TextureNameProvider;
+
+		for (int32 TextureIndex = 0, TextureNum = DatasmithScene->GetTexturesCount(); TextureIndex < TextureNum; ++TextureIndex)
+		{
+			if (TSharedPtr<IDatasmithTextureElement> TextureElement = DatasmithScene->GetTexture(TextureIndex))
+			{
+				UInterchangeTexture2DNode* TextureNode = NewObject<UInterchangeTexture2DNode>(&BaseNodeContainer);
+				const FString TextureNodeUid = NodeUtils::TexturePrefix + TextureElement->GetName();
+				const FString DisplayLabel = TextureNameProvider.GenerateUniqueName(TextureElement->GetLabel());
+
+				TextureNode->InitializeNode(TextureNodeUid, DisplayLabel, EInterchangeNodeContainerType::TranslatedAsset);
+				TextureNode->SetPayLoadKey(LexToString(TextureIndex));
+				TextureUtils::ApplyTextureElementToNode(TextureElement.ToSharedRef(), TextureNode);
+
+				BaseNodeContainer.AddNode(TextureNode);
+			}
+		}
+	}
+
+	// Materials
+	{
+		FDatasmithUniqueNameProvider MaterialsNameProvider;
+		const TCHAR* HostName = DatasmithScene->GetHost();
+
+		TArray<TSharedPtr<IDatasmithBaseMaterialElement>> MaterialElements;
+		MaterialElements.Reserve(DatasmithScene->GetMaterialsCount());
+
+		for (int32 MaterialIndex = 0, MaterialNum = DatasmithScene->GetMaterialsCount(); MaterialIndex < MaterialNum; ++MaterialIndex)
+		{
+			if (TSharedPtr<IDatasmithBaseMaterialElement> MaterialElement = DatasmithScene->GetMaterial(MaterialIndex))
+			{
+				MaterialElements.Add(MaterialElement);
+			}
+		}
+
+		MaterialUtils::ProcessMaterialElements(MaterialElements);
+
+		for (TSharedPtr<IDatasmithBaseMaterialElement>& MaterialElement : MaterialElements)
+		{
+			if (UInterchangeBaseNode* MaterialNode = MaterialUtils::AddMaterialNode(MaterialElement, BaseNodeContainer))
+			{
+				const FString DisplayLabel = MaterialsNameProvider.GenerateUniqueName(MaterialNode->GetDisplayLabel());
+				MaterialNode->SetDisplayLabel(DisplayLabel);
+
+				if (MaterialElement->IsA(EDatasmithElementType::MaterialInstance))
+				{
+					UInterchangeDatasmithMaterialNode* ReferenceMaterialNode = Cast<UInterchangeDatasmithMaterialNode>(MaterialNode);
+					if (ReferenceMaterialNode->GetMaterialType() == EDatasmithReferenceMaterialType::Custom)
+					{
+						ReferenceMaterialNode->SetParentPath(static_cast<IDatasmithMaterialInstanceElement&>(*MaterialElement).GetCustomMaterialPathName());
+					}
+					else
+					{
+						ReferenceMaterialNode->SetParentPath(HostName);
+					}
+				}
+			}
+		}
+	}
+
+	// Static Meshes
+	{
+		FDatasmithUniqueNameProvider StaticMeshNameProvider;
+		for (int32 MeshIndex = 0, MeshNum = DatasmithScene->GetMeshesCount(); MeshIndex < MeshNum; ++MeshIndex)
+		{
+			if (const TSharedPtr<IDatasmithMeshElement> MeshElement = DatasmithScene->GetMesh(MeshIndex))
+			{
+				UInterchangeMeshNode* MeshNode = NewObject<UInterchangeMeshNode>(&BaseNodeContainer);
+				const FString MeshNodeUid = NodeUtils::MeshPrefix + MeshElement->GetName();
+				const FString DisplayLabel = StaticMeshNameProvider.GenerateUniqueName(MeshElement->GetLabel());
+
+				MeshNode->InitializeNode(MeshNodeUid, DisplayLabel, EInterchangeNodeContainerType::TranslatedAsset);
+				MeshNode->SetPayLoadKey(LexToString(MeshIndex));
+				MeshNode->SetSkinnedMesh(false);
+				MeshNode->SetCustomHasVertexNormal(true);
+				// TODO: Interchange expect each LOD to have its own mesh node and to declare the number of vertices, however we don't know the content of a datasmith mesh until its bulk data is loaded.
+				//       It is not clear what would be the proper way to properly translate the content of the Datasmith meshes without	loading all this data during the translation phase (done on the main thread).
+
+				TSharedPtr<IDatasmithMaterialIDElement> GlobalMaterialID;
+				for (int32 SlotIndex = 0, SlotCount = MeshElement->GetMaterialSlotCount(); SlotIndex < SlotCount; ++SlotIndex)
+				{
+					if (TSharedPtr<IDatasmithMaterialIDElement> MaterialID = MeshElement->GetMaterialSlotAt(SlotIndex))
+					{
+						if (MaterialID->GetId() == -1)
+						{
+							GlobalMaterialID = MaterialID;
+							break;
+						}
+					}
+				}
+
+				if (GlobalMaterialID)
+				{
+					// Set dedicated attribute with value of material Uid.
+					// Corresponding factory then mesh asset will be updated accordingly pre then post import in the pipeline
+					const FString MaterialUid = NodeUtils::MaterialPrefix + GlobalMaterialID->GetName();
+					MeshNode->AddStringAttribute(MeshUtils::MeshMaterialAttrName, MaterialUid);
+				}
+				else
+				{
+					for (int32 SlotIndex = 0, SlotCount = MeshElement->GetMaterialSlotCount(); SlotIndex < SlotCount; ++SlotIndex)
+					{
+						if (const TSharedPtr<IDatasmithMaterialIDElement> MaterialID = MeshElement->GetMaterialSlotAt(SlotIndex))
+						{
+							const FString MaterialUid = NodeUtils::MaterialPrefix + MaterialID->GetName();
+							if (BaseNodeContainer.GetNode(MaterialUid) != nullptr)
+							{
+								MeshNode->SetSlotMaterialDependencyUid(*FString::FromInt(MaterialID->GetId()), MaterialUid);
+							}
+						}
+					}
+				}
+
+				BaseNodeContainer.AddNode(MeshNode);
+			}
+		}
+	}
+
+	//Actors
+	{
+		// Add base scene node.
+		UInterchangeSceneNode* SceneNode = NewObject< UInterchangeSceneNode >(&BaseNodeContainer);
+		const FString SceneName = DatasmithScene->GetName();
+		const FString SceneNodeUid = NodeUtils::ScenePrefix + SceneName;
+		SceneNode->InitializeNode(SceneNodeUid, DatasmithScene->GetLabel(), EInterchangeNodeContainerType::TranslatedScene);
+		// TODO: This should be the instantiation of the DatasmithScene asset, and create a DatasmithSceneActor.
+		
+		BaseNodeContainer.AddNode(SceneNode);
+
+		for (int32 ActorIndex = 0, ActorNum = DatasmithScene->GetActorsCount(); ActorIndex < ActorNum; ++ActorIndex)
+		{
+			if (const TSharedPtr<IDatasmithActorElement> ActorElement = DatasmithScene->GetActor(ActorIndex))
+			{
+				HandleDatasmithActor(BaseNodeContainer, ActorElement.ToSharedRef(), SceneNode);
+			}
+		}
+	}
+
+	// Level sequences
+	{
+		const int32 SequencesCount = DatasmithScene->GetLevelSequencesCount();
+		
+		TArray<TSharedPtr<IDatasmithLevelSequenceElement>> LevelSequences;
+		LevelSequences.Reserve(SequencesCount);
+
+
+		for (int32 SequenceIndex = 0; SequenceIndex < SequencesCount; ++SequenceIndex)
+		{
+			TSharedPtr<IDatasmithLevelSequenceElement> SequenceElement = DatasmithScene->GetLevelSequence(SequenceIndex);
+			if (!SequenceElement)
+			{
+				continue;
+			}
+
+			FDatasmithLevelSequencePayload LevelSequencePayload;
+			LoadedExternalSource->GetAssetTranslator()->LoadLevelSequence(SequenceElement.ToSharedRef(), LevelSequencePayload);
+
+			if (SequenceElement->GetAnimationsCount() > 0)
+			{
+				LevelSequences.Add(SequenceElement);
+			}
+		}
+
+		AnimUtils::TranslateLevelSequences(LevelSequences, BaseNodeContainer, AnimationPayLoadMapping);
+	}
+
+	// Level variant sets
+	{
+		const int32 LevelVariantSetCount = DatasmithScene->GetLevelVariantSetsCount();
+
+		TArray<TSharedPtr<IDatasmithLevelVariantSetsElement>> LevelVariantSets;
+		LevelVariantSets.Reserve(LevelVariantSetCount);
+
+
+		for (int32 LevelVariantSetIndex = 0; LevelVariantSetIndex < LevelVariantSetCount; ++LevelVariantSetIndex)
+		{
+			TSharedPtr<IDatasmithLevelVariantSetsElement> LevelVariantSetElement = DatasmithScene->GetLevelVariantSets(LevelVariantSetIndex);
+			if (LevelVariantSetElement)
+			{
+				LevelVariantSets.Add(LevelVariantSetElement);
+			}
+		}
+
+		VariantSetUtils::TranslateLevelVariantSets(LevelVariantSets, BaseNodeContainer);
+	}
+
+	return true;
+}
+
+void UInterchangeDatasmithTranslator::HandleDatasmithActor(UInterchangeBaseNodeContainer& BaseNodeContainer, const TSharedRef<IDatasmithActorElement>& ActorElement, const UInterchangeSceneNode* ParentNode) const
+{
+	using namespace UE::DatasmithInterchange;
+	const FString NodeName = ActorElement->GetName();
+	const FString ParentNodeUid = ParentNode->GetUniqueID();
+	const FString NodeUid = NodeUtils::GetActorUid(*NodeName);
+
+	UInterchangeSceneNode* InterchangeSceneNode = NewObject<UInterchangeSceneNode>(&BaseNodeContainer);
+	InterchangeSceneNode->InitializeNode(NodeUid, ActorElement->GetLabel(), EInterchangeNodeContainerType::TranslatedScene);
+	BaseNodeContainer.AddNode(InterchangeSceneNode);
+	BaseNodeContainer.SetNodeParentUid(NodeUid, ParentNodeUid);
+
+	const FTransform ActorTransform = ActorElement->GetRelativeTransform();
+	InterchangeSceneNode->SetCustomLocalTransform(&BaseNodeContainer, ActorElement->GetRelativeTransform());
+	// TODO: Layer association + component actors
+
+	if (ActorElement->IsA(EDatasmithElementType::StaticMeshActor))
+	{
+		TSharedRef<IDatasmithMeshActorElement> MeshActor = StaticCastSharedRef<IDatasmithMeshActorElement>(ActorElement);
+		// TODO: GetStaticMeshPathName() might reference an asset that was not imported.
+		const FString MeshUid = NodeUtils::MeshPrefix + MeshActor->GetStaticMeshPathName();
+		InterchangeSceneNode->SetCustomAssetInstanceUid(MeshUid);
+
+		TSharedPtr<IDatasmithMaterialIDElement> GlobalMaterialID;
+		for (int32 OverrideIndex = 0, OverrideCount = MeshActor->GetMaterialOverridesCount(); OverrideIndex < OverrideCount; ++OverrideIndex)
+		{
+			if (const TSharedPtr<IDatasmithMaterialIDElement> MaterialID = MeshActor->GetMaterialOverride(OverrideIndex))
+			{
+				if (MaterialID->GetId() == -1)
+				{
+					GlobalMaterialID = MaterialID;
+					break;
+				}
+			}
+		}
+
+		if (GlobalMaterialID)
+		{
+			// Set dedicated attribute with value of material Uid.
+			// Corresponding factory then mesh actor will be updated accordingly pre then post import in the pipeline
+			const FString MaterialUid = NodeUtils::MaterialPrefix + GlobalMaterialID->GetName();
+			InterchangeSceneNode->AddStringAttribute(MeshUtils::MeshMaterialAttrName, MaterialUid);
+		}
+		else
+		{
+			for (int32 OverrideIndex = 0, OverrideCount = MeshActor->GetMaterialOverridesCount(); OverrideIndex < OverrideCount; ++OverrideIndex)
+			{
+				if (const TSharedPtr<IDatasmithMaterialIDElement> MaterialID = MeshActor->GetMaterialOverride(OverrideIndex))
+				{
+					const FString MaterialUid = NodeUtils::MaterialPrefix + FDatasmithUtils::SanitizeObjectName(MaterialID->GetName());
+					if (BaseNodeContainer.GetNode(MaterialUid) != nullptr)
+					{
+						InterchangeSceneNode->SetSlotMaterialDependencyUid(*FString::FromInt(MaterialID->GetId()), MaterialUid);
+					}
+				}
+			}
+		}
+	}
+	else if (ActorElement->IsA(EDatasmithElementType::Camera))
+	{
+		TSharedRef<IDatasmithCameraActorElement> CameraActor = StaticCastSharedRef<IDatasmithCameraActorElement>(ActorElement);
+
+		// We need to add camera asset node and then instance it in the scene node.
+		UInterchangeCameraNode* CameraNode = AddCameraNode(BaseNodeContainer, CameraActor);
+		InterchangeSceneNode->SetCustomAssetInstanceUid(CameraNode->GetUniqueID());
+	}
+	else if (ActorElement->IsA(EDatasmithElementType::Light))
+	{
+		TSharedRef<IDatasmithLightActorElement> LightActor = StaticCastSharedRef<IDatasmithLightActorElement>(ActorElement);
+
+		// We need to add light asset node and then instance it in the scene node.
+		UInterchangeBaseLightNode* LightNode = AddLightNode(BaseNodeContainer, LightActor);
+		InterchangeSceneNode->SetCustomAssetInstanceUid(LightNode->GetUniqueID());
+	}
+
+	for (int32 ChildIndex = 0, ChildrenCount = ActorElement->GetChildrenCount(); ChildIndex < ChildrenCount; ++ChildIndex)
+	{
+		if (const TSharedPtr<IDatasmithActorElement> ChildActorElement = ActorElement->GetChild(ChildIndex))
+		{
+			HandleDatasmithActor(BaseNodeContainer, ChildActorElement.ToSharedRef(), InterchangeSceneNode);
+		}
+	}
+}
+
+UInterchangeCameraNode* UInterchangeDatasmithTranslator::AddCameraNode(UInterchangeBaseNodeContainer& BaseNodeContainer, const TSharedRef<IDatasmithCameraActorElement>& CameraActor) const
+{
+	using namespace UE::DatasmithInterchange;
+
+	UInterchangeCameraNode* CameraNode = NewObject<UInterchangeCameraNode>(&BaseNodeContainer);
+	const FString CameraUid = NodeUtils::CameraPrefix + CameraActor->GetName();
+	CameraNode->InitializeNode(CameraUid, CameraActor->GetLabel(), EInterchangeNodeContainerType::TranslatedAsset);
+	BaseNodeContainer.AddNode(CameraNode);
+
+	CameraNode->SetCustomFocalLength(CameraActor->GetFocalLength());
+	CameraNode->SetCustomSensorWidth(CameraActor->GetSensorWidth());
+	const float SensorHeight = CameraActor->GetSensorWidth() / CameraActor->GetSensorAspectRatio();
+	CameraNode->SetCustomSensorHeight(SensorHeight);
+
+	// TODO Add properties currently missing from the InterchangeCameraNode:
+	//  - DepthOfField
+	//  - FocusDistance
+	//  - FStop
+	//  - FocalLength
+	//  - PostProcess
+	//  - LookAtActor
+
+	return CameraNode;
+}
+
+UInterchangeBaseLightNode* UInterchangeDatasmithTranslator::AddLightNode(UInterchangeBaseNodeContainer& BaseNodeContainer, const TSharedRef<IDatasmithLightActorElement>& LightActor) const
+{
+	using namespace UE::DatasmithInterchange;
+
+	using FDatasmithLightUnits = std::underlying_type_t<EDatasmithLightUnits>;
+	using FInterchangeLightUnits = std::underlying_type_t<EInterchangeLightUnits>;
+	using FCommonLightUnits = std::common_type_t<FDatasmithLightUnits, FInterchangeLightUnits>;
+
+	static_assert(FCommonLightUnits(EInterchangeLightUnits::Unitless) == FCommonLightUnits(EDatasmithLightUnits::Unitless), "EDatasmithLightUnits::Unitless differs from EInterchangeLightUnits::Unitless");
+	static_assert(FCommonLightUnits(EInterchangeLightUnits::Lumens) == FCommonLightUnits(EDatasmithLightUnits::Lumens), "EDatasmithLightUnits::Lumens differs from EInterchangeLightUnits::Lumens");
+	static_assert(FCommonLightUnits(EInterchangeLightUnits::Candelas) == FCommonLightUnits(EDatasmithLightUnits::Candelas), "EDatasmithLightUnits::Candelas differs from EInterchangeLightUnits::Candelas");
+
+	// TODO Add properties currently missing from the UInterchangeLightNode: everything
+	UInterchangeBaseLightNode* LightNode = nullptr;
+	if (LightActor->IsA(EDatasmithElementType::AreaLight))
+	{
+		const TSharedRef<IDatasmithAreaLightElement> AreaLightElement = StaticCastSharedRef<IDatasmithAreaLightElement>(LightActor);
+		UInterchangeDatasmithAreaLightNode* AreaLightNode = NewObject<UInterchangeDatasmithAreaLightNode>(&BaseNodeContainer);
+
+		const FString LightUid = NodeUtils::LightPrefix + LightActor->GetName();
+		AreaLightNode->InitializeNode(LightUid, LightActor->GetLabel(), EInterchangeNodeContainerType::TranslatedAsset);
+
+		AreaLightNode->SetCustomLightType(static_cast<EDatasmithAreaLightActorType>(AreaLightElement->GetLightType()));
+		AreaLightNode->SetCustomLightShape(static_cast<EDatasmithAreaLightActorShape>(AreaLightElement->GetLightShape()));
+		AreaLightNode->SetCustomDimensions(FVector2D(AreaLightElement->GetLength(), AreaLightElement->GetWidth()));
+		AreaLightNode->SetCustomIntensity(AreaLightElement->GetIntensity());
+		AreaLightNode->SetCustomIntensityUnits(static_cast<EInterchangeLightUnits>(AreaLightElement->GetIntensityUnits()));
+		AreaLightNode->SetCustomColor(AreaLightElement->GetColor());
+		if (AreaLightElement->GetUseTemperature())
+		{
+			AreaLightNode->SetCustomTemperature(AreaLightElement->GetTemperature());
+		}
+
+		if (AreaLightElement->GetUseIes())
+		{
+			//AreaLightNode->SetCustomIESTexture(const TObjectPtr<class UTextureLightProfile>&AttributeValue, bool bAddApplyDelegate = true);
+			AreaLightNode->SetCustomUseIESBrightness(AreaLightElement->GetUseIesBrightness());
+			AreaLightNode->SetCustomIESBrightnessScale(AreaLightElement->GetIesBrightnessScale());
+			AreaLightNode->SetCustomRotation(AreaLightElement->GetIesRotation().Rotator());
+		}
+
+		AreaLightNode->SetCustomSourceRadius(AreaLightElement->GetSourceRadius());
+		AreaLightNode->SetCustomSourceLength(AreaLightElement->GetSourceLength());
+		AreaLightNode->SetCustomAttenuationRadius(AreaLightElement->GetAttenuationRadius());
+		AreaLightNode->SetCustomSpotlightInnerAngle(AreaLightElement->GetInnerConeAngle());
+		AreaLightNode->SetCustomSpotlightOuterAngle(AreaLightElement->GetOuterConeAngle());
+
+		LightNode = AreaLightNode;
+		BaseNodeContainer.AddNode(LightNode);
+		return LightNode;
+	}
+	else if (LightActor->IsA(EDatasmithElementType::SpotLight))
+	{
+		LightNode = NewObject<UInterchangeSpotLightNode>(&BaseNodeContainer);
+	}
+	else if (LightActor->IsA(EDatasmithElementType::LightmassPortal))
+	{
+		// TODO Add lightmass portal support in interchange.
+		LightNode = NewObject<UInterchangeRectLightNode>(&BaseNodeContainer);
+	}
+	else if (LightActor->IsA(EDatasmithElementType::PointLight))
+	{
+		LightNode = NewObject<UInterchangePointLightNode>(&BaseNodeContainer);
+	}
+	else if (LightActor->IsA(EDatasmithElementType::DirectionalLight))
+	{
+		LightNode = NewObject<UInterchangeDirectionalLightNode>(&BaseNodeContainer);
+	}
+	else
+	{
+		ensure(false);
+		LightNode = NewObject<UInterchangeLightNode>(&BaseNodeContainer);
+	}
+	const FString LightUid = NodeUtils::LightPrefix + LightActor->GetName();
+	LightNode->InitializeNode(LightUid, LightActor->GetLabel(), EInterchangeNodeContainerType::TranslatedAsset);
+	BaseNodeContainer.AddNode(LightNode);
+
+	return LightNode;
+}
+
+TOptional<UE::Interchange::FImportImage> UInterchangeDatasmithTranslator::GetTexturePayloadData(const UInterchangeSourceData* InPayloadSourceData, const FString& PayloadKey) const
+{
+	if (!LoadedExternalSource || !LoadedExternalSource->GetDatasmithScene())
+	{
+		return TOptional<UE::Interchange::FImportImage>();
+	}
+
+	int32 TextureIndex = 0;
+	LexFromString(TextureIndex, *PayloadKey);
+	TSharedPtr<IDatasmithScene> DatasmithScene = LoadedExternalSource->GetDatasmithScene();
+	if (TextureIndex < 0 || TextureIndex >= DatasmithScene->GetTexturesCount())
+	{
+		return TOptional<UE::Interchange::FImportImage>();
+	}
+
+	TSharedPtr<IDatasmithTextureElement> TextureElement = DatasmithScene->GetTexture(TextureIndex);
+	if (!TextureElement.IsValid())
+	{
+		return TOptional<UE::Interchange::FImportImage>();
+	}
+
+	UInterchangeSourceData* PayloadSourceData = UInterchangeManager::GetInterchangeManager().CreateSourceData(TextureElement->GetFile());
+	FGCObjectScopeGuard ScopedSourceData(PayloadSourceData);
+
+	if (!PayloadSourceData)
+	{
+		return TOptional<UE::Interchange::FImportImage>();
+	}
+
+	UInterchangeTranslatorBase* SourceTranslator = UInterchangeManager::GetInterchangeManager().GetTranslatorForSourceData(PayloadSourceData);
+	FGCObjectScopeGuard ScopedSourceTranslator(SourceTranslator);
+	const IInterchangeTexturePayloadInterface* TextureTranslator = Cast< IInterchangeTexturePayloadInterface >(SourceTranslator);
+	if (!ensure(TextureTranslator))
+	{
+		return TOptional<UE::Interchange::FImportImage>();
+	}
+
+	return TextureTranslator->GetTexturePayloadData(PayloadSourceData, TextureElement->GetFile());
+}
+
+TFuture<TOptional<UE::Interchange::FStaticMeshPayloadData>> UInterchangeDatasmithTranslator::GetStaticMeshPayloadData(const FString& PayloadKey) const
+{
+	TPromise<TOptional<UE::Interchange::FStaticMeshPayloadData>> EmptyPromise;
+	EmptyPromise.SetValue(TOptional<UE::Interchange::FStaticMeshPayloadData>());
+
+	if (!LoadedExternalSource || !LoadedExternalSource->GetDatasmithScene())
+	{
+		return EmptyPromise.GetFuture();
+	}
+
+	int32 MeshIndex = 0;
+	LexFromString(MeshIndex, *PayloadKey);
+	TSharedPtr<IDatasmithScene> DatasmithScene = LoadedExternalSource->GetDatasmithScene();
+	if (MeshIndex < 0 || MeshIndex >= DatasmithScene->GetMeshesCount())
+	{
+		return EmptyPromise.GetFuture();
+	}
+
+	TSharedPtr<IDatasmithMeshElement> MeshElement = DatasmithScene->GetMesh(MeshIndex);
+	if (!MeshElement.IsValid())
+	{
+		return EmptyPromise.GetFuture();
+	}
+
+	return Async(EAsyncExecution::TaskGraph, [this, MeshElement = MoveTemp(MeshElement)]
+		{
+			TOptional<UE::Interchange::FStaticMeshPayloadData> Result;
+
+			FDatasmithMeshElementPayload DatasmithMeshPayload;
+			if (LoadedExternalSource->GetAssetTranslator()->LoadStaticMesh(MeshElement.ToSharedRef(), DatasmithMeshPayload))
+			{
+				if (DatasmithMeshPayload.LodMeshes.Num() > 0)
+				{
+					UE::Interchange::FStaticMeshPayloadData StaticMeshPayloadData;
+					StaticMeshPayloadData.MeshDescription = MoveTemp(DatasmithMeshPayload.LodMeshes[0]);
+
+					Result.Emplace(MoveTemp(StaticMeshPayloadData));
+				}
+			}
+
+			return Result;
+		}
+	);
+}
+
+TFuture<TOptional<UE::Interchange::FAnimationCurvePayloadData>> UInterchangeDatasmithTranslator::GetAnimationCurvePayloadData(const FString& PayLoadKey) const
+{
+	return GetAnimationCurvePayloadData<UE::Interchange::FAnimationCurvePayloadData>(PayLoadKey);
+}
+
+TFuture<TOptional<UE::Interchange::FAnimationStepCurvePayloadData>> UInterchangeDatasmithTranslator::GetAnimationStepCurvePayloadData(const FString& PayLoadKey) const
+{
+	return GetAnimationCurvePayloadData<UE::Interchange::FAnimationStepCurvePayloadData>(PayLoadKey);
+}
+
+TFuture<TOptional<UE::Interchange::FAnimationBakeTransformPayloadData>> UInterchangeDatasmithTranslator::GetAnimationBakeTransformPayloadData(const FString& PayLoadKey, const double BakeFrequency, const double RangeStartSecond, const double RangeStopSecond) const
+{
+	TPromise<TOptional<UE::Interchange::FAnimationBakeTransformPayloadData>> EmptyPromise;
+	EmptyPromise.SetValue(TOptional<UE::Interchange::FAnimationBakeTransformPayloadData>());
+	return EmptyPromise.GetFuture();
+}
+
+TFuture<TOptional<UE::Interchange::FVariantSetPayloadData>> UInterchangeDatasmithTranslator::GetVariantSetPayloadData(const FString& PayloadKey) const
+{
+	using namespace UE::Interchange;
+	using namespace UE::DatasmithInterchange;
+
+	TPromise<TOptional<FVariantSetPayloadData>> EmptyPromise;
+	EmptyPromise.SetValue(TOptional<FVariantSetPayloadData>());
+
+	if (!LoadedExternalSource || !LoadedExternalSource->GetDatasmithScene())
+	{
+		return EmptyPromise.GetFuture();
+	}
+
+	TSharedPtr<IDatasmithScene> DatasmithScene = LoadedExternalSource->GetDatasmithScene();
+	
+	TArray<FString> PayloadTokens;
+
+	// We need two indices to build the payload: index of LevelVariantSet and index of VariantSetIndex
+	if (2 != PayloadKey.ParseIntoArray(PayloadTokens, TEXT(";")))
+	{
+		// Invalid payload
+		return EmptyPromise.GetFuture();
+	}
+
+	int32 LevelVariantSetIndex = FCString::Atoi(*PayloadTokens[0]);
+	int32 VariantSetIndex = FCString::Atoi(*PayloadTokens[1]);
+
+	TSharedPtr<IDatasmithLevelVariantSetsElement> LevelVariantSetElement = DatasmithScene->GetLevelVariantSets(LevelVariantSetIndex);
+	if (ensure(LevelVariantSetElement))
+	{
+		TSharedPtr<IDatasmithVariantSetElement> VariantSet = LevelVariantSetElement->GetVariantSet(VariantSetIndex);
+		if (ensure(VariantSet) && VariantSet->GetVariantsCount() > 0)
+		{
+			return Async(EAsyncExecution::TaskGraph, [this, VariantSet = MoveTemp(VariantSet)]
+					{
+						FVariantSetPayloadData PayloadData;
+						TOptional<FVariantSetPayloadData> Result;
+
+						if (VariantSetUtils::GetVariantSetPayloadData(*VariantSet, PayloadData))
+						{
+							Result.Emplace(MoveTemp(PayloadData));
+						}
+
+						return Result;
+					}
+				);
+		}
+	}
+
+	return EmptyPromise.GetFuture();
+}
+
+void UInterchangeDatasmithTranslator::ImportFinish()
+{
+	double ElapsedSeconds = FPlatformTime::ToSeconds64(FPlatformTime::Cycles64() - StartTime);
+
+	int ElapsedMin = int(ElapsedSeconds / 60.0);
+	ElapsedSeconds -= 60.0 * (double)ElapsedMin;
+
+	UE_LOG(LogInterchangeDatasmith, Log, TEXT("Imported %s in [%d min %.3f s]"), *FileName, ElapsedMin, ElapsedSeconds);
+}
+
+#undef LOCTEXT_NAMESPACE
\ No newline at end of file
diff --git a/Engine/Plugins/Experimental/Enterprise/DatasmithInterchange/Source/DatasmithInterchange/Private/InterchangeDatasmithUtils.cpp b/Engine/Plugins/Experimental/Enterprise/DatasmithInterchange/Source/DatasmithInterchange/Private/InterchangeDatasmithUtils.cpp
new file mode 100644
--- /dev/null
+++ b/Engine/Plugins/Experimental/Enterprise/DatasmithInterchange/Source/DatasmithInterchange/Private/InterchangeDatasmithUtils.cpp
@@ -0,0 +1,263 @@
+// Copyright Epic Games, Inc. All Rights Reserved.
+
+#include "InterchangeDatasmithUtils.h"
+
+#include "InterchangeVariantSetNode.h"
+#include "Nodes/InterchangeFactoryBaseNode.h"
+#include "Scene/InterchangeVariantSetPayloadInterface.h"
+
+#include "DatasmithVariantElements.h"
+
+#include "Templates/SharedPointer.h"
+
+namespace UE::DatasmithInterchange::NodeUtils
+{
+	const FString ActorPrefix(FString(TEXT("\\Actor\\")));
+	const FString CameraPrefix(FString(TEXT("\\Camera\\")));
+	const FString DatasmithScenePrefix(FString(TEXT("\\DatasmithScene\\")));
+	const FString LevelSequencePrefix(FString(TEXT("\\LevelSequence\\")));
+	const FString LevelVariantSetPrefix(FString(TEXT("\\LevelVariantSet\\")));
+	const FString LightPrefix(FString(TEXT("\\Light\\")));
+	const FString MaterialPrefix(FString(TEXT("\\Material\\")));
+	const FString MaterialExpressionPrefix(FString(TEXT("\\MaterialExpression\\")));
+	const FString MaterialFunctionPrefix(FString(TEXT("\\MaterialFunction\\")));
+	const FString MeshPrefix(FString(TEXT("\\StaticMesh\\")));
+	const FString ScenePrefix(FString(TEXT("\\Scene\\")));
+	const FString TexturePrefix(FString(TEXT("\\Texture\\")));
+	const FString VariantSetPrefix(FString(TEXT("\\VariantSet\\")));
+
+	UInterchangeFactoryBaseNode* FindFactoryNodeFromAsset(const UInterchangeBaseNodeContainer* NodeContainer, const UObject* Asset)
+	{
+		UInterchangeFactoryBaseNode* FactoryNode = nullptr;
+		NodeContainer->BreakableIterateNodesOfType<UInterchangeFactoryBaseNode>([&Asset, &FactoryNode](const FString& NodeUid, UInterchangeFactoryBaseNode* Node)
+			{
+				FSoftObjectPath NodeReferenceObject;
+				Node->GetCustomReferenceObject(NodeReferenceObject);
+				if (NodeReferenceObject.TryLoad() == Asset)
+				{
+					FactoryNode = Node;
+					return true;
+				}
+				return false;
+			});
+
+		return FactoryNode;
+	}
+
+	FString GetActorUid(const TCHAR* Name)
+	{
+		return ActorPrefix + Name;
+	}
+
+	FString GetLevelSequenceUid(const TCHAR* Name)
+	{
+		return LevelSequencePrefix + Name;
+	}
+
+	FString GetLevelVariantSetUid(const TCHAR* Name)
+	{
+		return LevelVariantSetPrefix + Name;
+	}
+
+	FString GetMaterialUid(const TCHAR* Name)
+	{
+		return MaterialPrefix + Name;
+	}
+
+	FString GetVariantSetUid(const TCHAR* Name)
+	{
+		return VariantSetPrefix + Name;
+	}
+}
+
+namespace UE::DatasmithInterchange::MeshUtils
+{
+	const FName MeshMaterialAttrName(TEXT("Datasmith:Mesh:MaterialOverride"));
+}
+
+namespace UE::DatasmithInterchange::VariantSetUtils
+{
+	void TranslateLevelVariantSets(const TArray<TSharedPtr<IDatasmithLevelVariantSetsElement>>& LevelVariantSets, UInterchangeBaseNodeContainer& BaseNodeContainer)
+	{
+		int32 LevelVariantSetIndex = -1;
+		for (const TSharedPtr<IDatasmithLevelVariantSetsElement>& LevelVariantSet : LevelVariantSets)
+		{
+			++LevelVariantSetIndex;
+
+			if (!LevelVariantSet || LevelVariantSet->GetVariantSetsCount() == 0)
+			{
+				continue;
+			}
+
+			UInterchangeSceneVariantSetsNode* LevelVariantSetNode = NewObject< UInterchangeSceneVariantSetsNode >(&BaseNodeContainer);
+			if (!ensure(LevelVariantSetNode))
+			{
+				continue;
+			}
+
+			const FString LevelVariantSetNodeUid = NodeUtils::GetLevelVariantSetUid(LevelVariantSet->GetName());
+			LevelVariantSetNode->InitializeNode(LevelVariantSetNodeUid, LevelVariantSet->GetLabel(), EInterchangeNodeContainerType::TranslatedScene);
+
+			for (int32 VariantSetIndex = 0; VariantSetIndex < LevelVariantSet->GetVariantSetsCount(); ++VariantSetIndex)
+			{
+				TSharedPtr<IDatasmithVariantSetElement> VariantSet = LevelVariantSet->GetVariantSet(VariantSetIndex);
+				if (!VariantSet || VariantSet->GetVariantsCount() == 0)
+				{
+					continue;
+				}
+
+				UInterchangeVariantSetNode* VariantSetNode = NewObject< UInterchangeVariantSetNode >(&BaseNodeContainer);
+				if (!ensure(VariantSetNode))
+				{
+					continue;
+				}
+
+				const FString VariantSetNodeUid = NodeUtils::GetVariantSetUid(VariantSet->GetName());
+				VariantSetNode->InitializeNode(VariantSetNodeUid, VariantSet->GetLabel(), EInterchangeNodeContainerType::TranslatedScene);
+
+				for (int32 VariantIndex = 0; VariantIndex < VariantSet->GetVariantsCount(); ++VariantIndex)
+				{
+					if (TSharedPtr<IDatasmithVariantElement> Variant = VariantSet->GetVariant(VariantIndex))
+					{
+						for (int32 BindingIndex = 0; BindingIndex < Variant->GetActorBindingsCount(); ++BindingIndex)
+						{
+							TSharedPtr<IDatasmithActorBindingElement> Binding = Variant->GetActorBinding(BindingIndex);
+							if (!Binding || Binding->GetPropertyCapturesCount() == 0)
+							{
+								continue;
+							}
+
+							if (TSharedPtr<IDatasmithActorElement> Actor = Binding->GetActor())
+							{
+								bool bMustImport = false;
+
+								for (int32 PropertyCapturIndex = 0; PropertyCapturIndex < Binding->GetPropertyCapturesCount(); ++PropertyCapturIndex)
+								{
+									if (TSharedPtr<IDatasmithBasePropertyCaptureElement> BasePropertyCapture = Binding->GetPropertyCapture(PropertyCapturIndex))
+									{
+										// If property capture is an object, add its corresponding unique id as a dependency
+										if (BasePropertyCapture->IsSubType(EDatasmithElementVariantSubType::ObjectPropertyCapture))
+										{
+											TSharedPtr<IDatasmithObjectPropertyCaptureElement> PropertyCapture = StaticCastSharedRef<IDatasmithObjectPropertyCaptureElement>(BasePropertyCapture.ToSharedRef());
+
+											if (TSharedPtr<IDatasmithElement> TargetElement = PropertyCapture->GetRecordedObject().Pin())
+											{
+												if (TargetElement->IsA(EDatasmithElementType::BaseMaterial))
+												{
+													VariantSetNode->AddCustomDependencyUid(NodeUtils::GetMaterialUid(TargetElement->GetName()));
+												}
+											}
+										}
+
+										bMustImport = true;
+									}
+								}
+
+								if (bMustImport)
+								{
+									VariantSetNode->AddCustomDependencyUid(NodeUtils::GetActorUid(Actor->GetName()));
+								}
+							}
+						}
+					}
+				}
+
+				if (VariantSetNode->GetCustomDependencyUidCount() > 0)
+				{
+					VariantSetNode->SetCustomVariantsPayloadKey(FString::Printf(TEXT("%d;%d"), LevelVariantSetIndex, VariantSetIndex));
+
+					VariantSetNode->SetCustomDisplayText(VariantSet->GetName());
+
+					LevelVariantSetNode->AddCustomVariantSetUid(VariantSetNodeUid);
+					BaseNodeContainer.AddNode(VariantSetNode);
+				}
+			}
+
+			// Only add to container if there is something to import
+			if (LevelVariantSetNode->GetCustomVariantSetUidCount() > 0)
+			{
+				BaseNodeContainer.AddNode(LevelVariantSetNode);
+			}
+		}
+	}
+
+	bool GetVariantSetPayloadData(const IDatasmithVariantSetElement& VariantSet, UE::Interchange::FVariantSetPayloadData& PayLoadData)
+	{
+		using namespace UE::Interchange;
+
+		PayLoadData.Variants.Reserve(VariantSet.GetVariantsCount());
+
+		for (int32 VariantIndex = 0; VariantIndex < VariantSet.GetVariantsCount(); ++VariantIndex)
+		{
+			// Temporarily using const_cast since IDatasmithVariantSetElement::GetVariant does not have (yet) a const alternative.
+			if (const TSharedPtr<IDatasmithVariantElement> VariantElement = const_cast<IDatasmithVariantSetElement&>(VariantSet).GetVariant(VariantIndex))
+			{
+				TArray<FVariantBinding> Bindings;
+				Bindings.Reserve(VariantElement->GetActorBindingsCount());
+
+				for (int32 BindingIndex = 0; BindingIndex < VariantElement->GetActorBindingsCount(); ++BindingIndex)
+				{
+					TSharedPtr<IDatasmithActorBindingElement> BindingElement = VariantElement->GetActorBinding(BindingIndex);
+					if (!BindingElement || BindingElement->GetPropertyCapturesCount() == 0)
+					{
+						continue;
+					}
+
+					if (TSharedPtr<IDatasmithActorElement> Actor = BindingElement->GetActor())
+					{
+						TArray<FVariantPropertyCaptureData> Captures;
+						Captures.Reserve(BindingElement->GetPropertyCapturesCount());
+
+						for (int32 PropertyCapturIndex = 0; PropertyCapturIndex < BindingElement->GetPropertyCapturesCount(); ++PropertyCapturIndex)
+						{
+							if (TSharedPtr<IDatasmithBasePropertyCaptureElement> BasePropertyCapture = BindingElement->GetPropertyCapture(PropertyCapturIndex))
+							{
+								FVariantPropertyCaptureData& PropertyCaptureData = Captures.AddDefaulted_GetRef();
+
+								PropertyCaptureData.Category = (Interchange::EVariantPropertyCaptureCategory)BasePropertyCapture->GetCategory();
+								PropertyCaptureData.PropertyPath = BasePropertyCapture->GetPropertyPath();
+
+								// If property capture is an object, add its corresponding unique id as a dependency
+								if (BasePropertyCapture->IsSubType(EDatasmithElementVariantSubType::ObjectPropertyCapture))
+								{
+									TSharedPtr<IDatasmithObjectPropertyCaptureElement> PropertyCapture = StaticCastSharedRef<IDatasmithObjectPropertyCaptureElement>(BasePropertyCapture.ToSharedRef());
+
+									if (TSharedPtr<IDatasmithElement> TargetElement = PropertyCapture->GetRecordedObject().Pin())
+									{
+										if (TargetElement->IsA(EDatasmithElementType::BaseMaterial))
+										{
+											PropertyCaptureData.ObjectUid = NodeUtils::GetMaterialUid(TargetElement->GetName());
+										}
+									}
+								}
+								else
+								{
+									TSharedPtr<IDatasmithPropertyCaptureElement> PropertyCapture = StaticCastSharedRef<IDatasmithPropertyCaptureElement>(BasePropertyCapture.ToSharedRef());
+									PropertyCaptureData.Buffer = PropertyCapture->GetRecordedData();
+								}
+							}
+						}
+
+						if (Captures.Num() > 0)
+						{
+							FVariantBinding& Binding = Bindings.AddDefaulted_GetRef();
+
+							Binding.TargetUid = NodeUtils::GetActorUid(Actor->GetName());
+							Binding.Captures = MoveTemp(Captures);
+						}
+					}
+				}
+
+				if (Bindings.Num() > 0)
+				{
+					FVariant& Variant = PayLoadData.Variants.AddDefaulted_GetRef();
+
+					Variant.DisplayText = VariantElement->GetName();
+					Variant.Bindings = MoveTemp(Bindings);
+				}
+			}
+		}
+
+		return PayLoadData.Variants.Num() > 0;
+	}
+}
diff --git a/Engine/Plugins/Experimental/Enterprise/DatasmithInterchange/Source/DatasmithInterchange/Private/InterchangeDatasmithUtils.h b/Engine/Plugins/Experimental/Enterprise/DatasmithInterchange/Source/DatasmithInterchange/Private/InterchangeDatasmithUtils.h
new file mode 100644
--- /dev/null
+++ b/Engine/Plugins/Experimental/Enterprise/DatasmithInterchange/Source/DatasmithInterchange/Private/InterchangeDatasmithUtils.h
@@ -0,0 +1,184 @@
+// Copyright Epic Games, Inc. All Rights Reserved.
+
+#pragma once
+
+#include "CoreMinimal.h"
+
+#include "Nodes/InterchangeBaseNode.h"
+#include "Nodes/InterchangeBaseNodeContainer.h"
+#include "Nodes/InterchangeFactoryBaseNode.h"
+
+#include "IDatasmithSceneElements.h"
+
+class IDatasmithBaseAnimationElement;
+class IDatasmithBaseMaterialElement;
+class IDatasmithLevelVariantSetsElement;
+class IDatasmithVariantSetElement;
+class UInterchangeDatasmithMaterialNode;
+
+namespace UE::Interchange
+{
+	struct FAnimationCurvePayloadData;
+	struct FAnimationStepCurvePayloadData;
+	struct FVariantSetPayloadData;
+}
+
+namespace UE::DatasmithInterchange::NodeUtils
+{
+	extern const FString ActorPrefix;
+	extern const FString CameraPrefix;
+	extern const FString DatasmithScenePrefix;
+	extern const FString LevelSequencePrefix;
+	extern const FString LevelVariantSetPrefix;
+	extern const FString LightPrefix;
+	extern const FString MaterialPrefix;
+	extern const FString MaterialExpressionPrefix;
+	extern const FString MaterialFunctionPrefix;
+	extern const FString MeshPrefix;
+	extern const FString ScenePrefix;
+	extern const FString TexturePrefix;
+	extern const FString VariantSetPrefix;
+
+	extern FString GetActorUid(const TCHAR* Name);
+	extern FString GetLevelSequenceUid(const TCHAR* Name);
+	extern FString GetLevelVariantSetUid(const TCHAR* Name);
+	extern FString GetMaterialUid(const TCHAR* Name);
+	extern FString GetVariantSetUid(const TCHAR* Name);
+
+	// TODO: This is just a quick helper to get the code started, the real solution would be to use caching and pre-parse the NodeContainer to avoid O(n^2) query.
+	template<typename T>
+	TArray<T*> GetNodes(const UInterchangeBaseNodeContainer* BaseNodeContainer)
+	{
+		TArray<T*> TypedNodes;
+
+		BaseNodeContainer->IterateNodes([&TypedNodes, &BaseNodeContainer](const FString& NodeUid, UInterchangeBaseNode* Node)
+		{
+			if (T* TypedNode = Cast<T>(Node))
+			{
+				TypedNodes.Add(TypedNode);
+			}
+		});
+
+		return TypedNodes;
+	}
+
+	template<class T>
+	T* FindOrAddFactoryNode(const UInterchangeBaseNode* Node, UInterchangeBaseNodeContainer* NodeContainer, const FString& FactoryNodeUid)
+	{
+		T* FactoryNode = nullptr;
+
+		if (NodeContainer->IsNodeUidValid(FactoryNodeUid))
+		{
+			//The node already exist, just return it
+			FactoryNode = Cast<T>(NodeContainer->GetFactoryNode(FactoryNodeUid));
+			if (!ensure(FactoryNode))
+			{
+				//Log an error
+			}
+		}
+		else
+		{
+			FactoryNode = NewObject<T>(NodeContainer, T::StaticClass(), NAME_None);
+			if (!ensure(FactoryNode))
+			{
+				return nullptr;
+			}
+
+			//Creating a Material
+			FactoryNode->InitializeNode(FactoryNodeUid, Node->GetDisplayLabel(), EInterchangeNodeContainerType::FactoryData);
+
+			NodeContainer->AddNode(FactoryNode);
+
+			FactoryNode->AddTargetNodeUid(Node->GetUniqueID());
+			FactoryNode->SetAssetName(Node->GetAssetName());
+
+			Node->AddTargetNodeUid(FactoryNode->GetUniqueID());
+		}
+
+		return FactoryNode;
+	}
+
+	UInterchangeFactoryBaseNode* FindFactoryNodeFromAsset(const UInterchangeBaseNodeContainer* NodeContainer, const UObject* Asset);
+
+	template<class T = UInterchangeFactoryBaseNode>
+	T* FindFactoryNodeByUniqueID(UInterchangeBaseNodeContainer* NodeContainer, const FString& NodeUniqueID)
+	{
+		T* FactoryNode = nullptr;
+		NodeContainer->BreakableIterateNodes([&NodeUniqueID, &FactoryNode](const FString& NodeUid, UInterchangeBaseNode* Node)
+			{
+				if (NodeUniqueID == NodeUid)
+				{
+					Cast<T>(Node);
+					return true;
+				}
+				return false;
+			});
+
+		ensure(FactoryNode);
+
+		return FactoryNode;
+	}
+
+
+	template<typename DatasmithElementType>
+	int32 GetDatasmithElementsCount(const TSharedRef<IDatasmithScene>& DatasmithScene); //no definition, template should be specialized
+
+	template<typename DatasmithElementType>
+	TSharedPtr<DatasmithElementType> GetDatasmithElement(const TSharedRef<IDatasmithScene>& DatasmithScene, int32 ElementIndex); //no definition, template should be specialized
+
+	template<typename DatasmithElementType>
+	TSharedPtr<DatasmithElementType> GetDatasmithElementFromIndex(const TSharedRef<IDatasmithScene>& DatasmithScene, int32 ElementIndex)
+	{
+		if (ElementIndex < 0 || ElementIndex >= GetDatasmithElementsCount<DatasmithElementType>(DatasmithScene))
+		{
+			return TSharedPtr<DatasmithElementType>();
+		}
+
+		return GetDatasmithElement<DatasmithElementType>(DatasmithScene, ElementIndex);
+	}
+
+	template<typename DatasmithElementType>
+	TSharedPtr<DatasmithElementType> GetDatasmithElementFromPayloadKey(const TSharedRef<IDatasmithScene>& DatasmithScene, const FString& PayloadKey)
+	{
+		int32 ElementIndex = 0;
+		LexFromString(ElementIndex, *PayloadKey);
+
+		return GetDatasmithElementFromIndex<DatasmithElementType>(DatasmithScene, ElementIndex);
+	}
+}
+
+namespace UE::DatasmithInterchange::MeshUtils
+{
+	extern const FName MeshMaterialAttrName;
+}
+
+namespace UE::DatasmithInterchange::AnimUtils
+{
+
+	typedef TPair<float, TSharedPtr<IDatasmithBaseAnimationElement>> FAnimationPayloadDesc;
+
+	extern void TranslateLevelSequences(TArray<TSharedPtr<IDatasmithLevelSequenceElement>>& LevelSequences, UInterchangeBaseNodeContainer& InBaseNodeContainer, TMap<FString, UE::DatasmithInterchange::AnimUtils::FAnimationPayloadDesc>& AnimationPayLoadMapping);
+	extern bool GetAnimationPayloadData(const IDatasmithBaseAnimationElement& AnimationElement, float FrameRate, UE::Interchange::FAnimationCurvePayloadData& PayLoadData);
+	extern bool GetAnimationPayloadData(const IDatasmithBaseAnimationElement& AnimationElement, float FrameRate, UE::Interchange::FAnimationStepCurvePayloadData& PayLoadData);
+}
+
+namespace UE::DatasmithInterchange::VariantSetUtils
+{
+
+	extern void TranslateLevelVariantSets(const TArray<TSharedPtr<IDatasmithLevelVariantSetsElement>>& LevelVariantSets, UInterchangeBaseNodeContainer& InBaseNodeContainer);
+	extern bool GetVariantSetPayloadData(const IDatasmithVariantSetElement& VariantSet, UE::Interchange::FVariantSetPayloadData& PayLoadData);
+}
+
+template<>
+inline int32 UE::DatasmithInterchange::NodeUtils::GetDatasmithElementsCount<IDatasmithTextureElement>(const TSharedRef<IDatasmithScene>& DatasmithScene) { return DatasmithScene->GetTexturesCount(); }
+template<>
+inline int32 UE::DatasmithInterchange::NodeUtils::GetDatasmithElementsCount<IDatasmithBaseMaterialElement>(const TSharedRef<IDatasmithScene>& DatasmithScene) { return DatasmithScene->GetMaterialsCount(); }
+template<>
+inline int32 UE::DatasmithInterchange::NodeUtils::GetDatasmithElementsCount<IDatasmithMeshElement>(const TSharedRef<IDatasmithScene>& DatasmithScene) { return DatasmithScene->GetMeshesCount(); }
+
+template<>
+inline TSharedPtr<IDatasmithTextureElement> UE::DatasmithInterchange::NodeUtils::GetDatasmithElement<IDatasmithTextureElement>(const TSharedRef<IDatasmithScene>& DatasmithScene, int32 ElementIndex) { return DatasmithScene->GetTexture(ElementIndex); }
+template<>
+inline TSharedPtr<IDatasmithBaseMaterialElement> UE::DatasmithInterchange::NodeUtils::GetDatasmithElement<IDatasmithBaseMaterialElement>(const TSharedRef<IDatasmithScene>& DatasmithScene, int32 ElementIndex) { return DatasmithScene->GetMaterial(ElementIndex); }
+template<>
+inline TSharedPtr<IDatasmithMeshElement> UE::DatasmithInterchange::NodeUtils::GetDatasmithElement<IDatasmithMeshElement>(const TSharedRef<IDatasmithScene>& DatasmithScene, int32 ElementIndex) { return DatasmithScene->GetMesh(ElementIndex); }
\ No newline at end of file
diff --git a/Engine/Plugins/Experimental/Enterprise/DatasmithInterchange/Source/DatasmithInterchange/Private/InterchangeReferenceMaterials/DatasmithC4DMaterialSelector.cpp b/Engine/Plugins/Experimental/Enterprise/DatasmithInterchange/Source/DatasmithInterchange/Private/InterchangeReferenceMaterials/DatasmithC4DMaterialSelector.cpp
new file mode 100644
--- /dev/null
+++ b/Engine/Plugins/Experimental/Enterprise/DatasmithInterchange/Source/DatasmithInterchange/Private/InterchangeReferenceMaterials/DatasmithC4DMaterialSelector.cpp
@@ -0,0 +1,53 @@
+// Copyright Epic Games, Inc. All Rights Reserved.
+
+#include "DatasmithC4DMaterialSelector.h"
+
+#include "IDatasmithSceneElements.h"
+
+#include "Materials/Material.h"
+#include "Materials/MaterialInstanceConstant.h"
+#include "Misc/PackageName.h"
+#include "UObject/SoftObjectPath.h"
+#include "Templates/Casts.h"
+
+namespace UE
+{
+	namespace DatasmithInterchange
+	{
+		namespace C4D
+		{
+			const TCHAR* MaterialPath = TEXT("/DatasmithContent/Materials/C4DReference.C4DReference");
+		}
+
+		FDatasmithC4DMaterialSelector::FDatasmithC4DMaterialSelector()
+		{
+			bIsValid = FPackageName::DoesPackageExist(C4D::MaterialPath);
+		}
+
+		const TCHAR* FDatasmithC4DMaterialSelector::GetMaterialPath(EDatasmithReferenceMaterialType) const
+		{
+			return C4D::MaterialPath;
+		}
+
+#if WITH_EDITOR
+		void FDatasmithC4DMaterialSelector::PostImportProcess(EDatasmithReferenceMaterialType MaterialType, EDatasmithReferenceMaterialQuality, UMaterialInstanceConstant* MaterialInstance) const
+		{
+			if (MaterialInstance)
+			{
+				// Set blend mode to translucent if material requires transparency.
+				if (MaterialType == EDatasmithReferenceMaterialType::Transparent)
+				{
+					MaterialInstance->BasePropertyOverrides.bOverride_BlendMode = true;
+					MaterialInstance->BasePropertyOverrides.BlendMode = BLEND_Translucent;
+				}
+				// Set blend mode to masked if material has cutouts.
+				else if (MaterialType == EDatasmithReferenceMaterialType::CutOut)
+				{
+					MaterialInstance->BasePropertyOverrides.bOverride_BlendMode = true;
+					MaterialInstance->BasePropertyOverrides.BlendMode = BLEND_Masked;
+				}
+			}
+		}
+#endif
+	}
+}
\ No newline at end of file
diff --git a/Engine/Plugins/Experimental/Enterprise/DatasmithInterchange/Source/DatasmithInterchange/Private/InterchangeReferenceMaterials/DatasmithC4DMaterialSelector.h b/Engine/Plugins/Experimental/Enterprise/DatasmithInterchange/Source/DatasmithInterchange/Private/InterchangeReferenceMaterials/DatasmithC4DMaterialSelector.h
new file mode 100644
--- /dev/null
+++ b/Engine/Plugins/Experimental/Enterprise/DatasmithInterchange/Source/DatasmithInterchange/Private/InterchangeReferenceMaterials/DatasmithC4DMaterialSelector.h
@@ -0,0 +1,27 @@
+// Copyright Epic Games, Inc. All Rights Reserved.
+
+#pragma once
+
+#include "DatasmithDefinitions.h"
+#include "InterchangeReferenceMaterials/DatasmithReferenceMaterialSelector.h"
+
+#include "Templates/SharedPointer.h"
+
+namespace UE
+{
+	namespace DatasmithInterchange
+	{
+		class FDatasmithC4DMaterialSelector : public FDatasmithReferenceMaterialSelector
+		{
+		public:
+
+			FDatasmithC4DMaterialSelector();
+			virtual ~FDatasmithC4DMaterialSelector() = default;
+
+			virtual const TCHAR* GetMaterialPath(EDatasmithReferenceMaterialType MaterialType) const override;
+#if WITH_EDITOR
+			virtual void PostImportProcess(EDatasmithReferenceMaterialType MaterialType, EDatasmithReferenceMaterialQuality MaterialQuality, UMaterialInstanceConstant* MaterialInstance) const override;
+#endif
+		};
+	}
+}
diff --git a/Engine/Plugins/Experimental/Enterprise/DatasmithInterchange/Source/DatasmithInterchange/Private/InterchangeReferenceMaterials/DatasmithCityEngineMaterialSelector.cpp b/Engine/Plugins/Experimental/Enterprise/DatasmithInterchange/Source/DatasmithInterchange/Private/InterchangeReferenceMaterials/DatasmithCityEngineMaterialSelector.cpp
new file mode 100644
--- /dev/null
+++ b/Engine/Plugins/Experimental/Enterprise/DatasmithInterchange/Source/DatasmithInterchange/Private/InterchangeReferenceMaterials/DatasmithCityEngineMaterialSelector.cpp
@@ -0,0 +1,91 @@
+// Copyright Epic Games, Inc. All Rights Reserved.
+
+#include "DatasmithCityEngineMaterialSelector.h"
+
+#include "IDatasmithSceneElements.h"
+
+#include "Materials/Material.h"
+#include "Materials/MaterialInstanceConstant.h"
+#include "UObject/SoftObjectPath.h"
+#include "Templates/Casts.h"
+
+namespace UE
+{
+	namespace DatasmithInterchange
+	{
+		namespace CityEngine
+		{
+			const TCHAR* MaterialPath = TEXT("/DatasmithContent/Materials/CE_OpaqueReference.CE_OpaqueReference");
+			const TCHAR* MaterialTransparentPath = TEXT("/DatasmithContent/Materials/CE_OpacityReference.CE_OpacityReference");
+			const TCHAR* MaterialTransparentSimplePath = TEXT("/DatasmithContent/Materials/CE_OpacitySimpleReference.CE_OpacitySimpleReference");
+		}
+
+		FDatasmithCityEngineMaterialSelector::FDatasmithCityEngineMaterialSelector()
+		{
+			bIsValid = FPackageName::DoesPackageExist(CityEngine::MaterialPath) || FPackageName::DoesPackageExist(CityEngine::MaterialTransparentPath) || FPackageName::DoesPackageExist(CityEngine::MaterialTransparentSimplePath);
+		}
+
+		const TCHAR* FDatasmithCityEngineMaterialSelector::GetMaterialPath(EDatasmithReferenceMaterialType MaterialType) const
+		{
+			switch (MaterialType)
+			{
+			case EDatasmithReferenceMaterialType::Transparent:
+				return CityEngine::MaterialTransparentPath;
+				break;
+			default:
+				return CityEngine::MaterialPath;
+				break;
+			}
+		}
+
+#if WITH_EDITOR
+		void FDatasmithCityEngineMaterialSelector::PostImportProcess(EDatasmithReferenceMaterialType MaterialType, EDatasmithReferenceMaterialQuality MaterialQuality, UMaterialInstanceConstant* MaterialInstance) const
+		{
+			if (MaterialInstance == nullptr || MaterialInstance->Parent == nullptr)
+			{
+				return;
+			}
+
+			if (MaterialInstance->Parent->GetFullName().Contains(CityEngine::MaterialTransparentPath))
+			{
+				if (MaterialQuality == EDatasmithReferenceMaterialQuality::Low)
+				{
+					MaterialInstance->Parent = Cast<UMaterialInterface>(FSoftObjectPath(CityEngine::MaterialTransparentSimplePath).TryLoad());
+					if (!ensure(MaterialInstance->Parent))
+					{
+						MaterialInstance->Parent = Cast<UMaterialInterface>(FSoftObjectPath(CityEngine::MaterialTransparentPath).TryLoad());
+					}
+				}
+			}
+			else
+			{
+				float OpacityValue = 1.f;
+				bool bIsTransparent = MaterialInstance->GetScalarParameterValue(TEXT("Opacity"), OpacityValue) && OpacityValue < 1.f;
+
+				class UTexture* TextureValue = nullptr;
+				MaterialInstance->GetTextureParameterValue(TEXT("OpacityMap"), TextureValue);
+				bIsTransparent |= TextureValue != nullptr;
+
+				if (bIsTransparent)
+				{
+					UMaterialInterface* PrevParent = MaterialInstance->Parent;
+
+					if (MaterialQuality == EDatasmithReferenceMaterialQuality::Low)
+					{
+						MaterialInstance->Parent = Cast<UMaterialInterface>(FSoftObjectPath(CityEngine::MaterialTransparentSimplePath).TryLoad());
+					}
+					else
+					{
+						MaterialInstance->Parent = Cast<UMaterialInterface>(FSoftObjectPath(CityEngine::MaterialTransparentPath).TryLoad());
+					}
+
+					if (!ensure(MaterialInstance->Parent))
+					{
+						MaterialInstance->Parent = PrevParent;
+					}
+				}
+			}
+		}
+#endif
+	}
+}
\ No newline at end of file
diff --git a/Engine/Plugins/Experimental/Enterprise/DatasmithInterchange/Source/DatasmithInterchange/Private/InterchangeReferenceMaterials/DatasmithCityEngineMaterialSelector.h b/Engine/Plugins/Experimental/Enterprise/DatasmithInterchange/Source/DatasmithInterchange/Private/InterchangeReferenceMaterials/DatasmithCityEngineMaterialSelector.h
new file mode 100644
--- /dev/null
+++ b/Engine/Plugins/Experimental/Enterprise/DatasmithInterchange/Source/DatasmithInterchange/Private/InterchangeReferenceMaterials/DatasmithCityEngineMaterialSelector.h
@@ -0,0 +1,24 @@
+// Copyright Epic Games, Inc. All Rights Reserved.
+#pragma once
+
+#include "DatasmithDefinitions.h"
+#include "InterchangeReferenceMaterials/DatasmithReferenceMaterialSelector.h"
+
+#include "Templates/SharedPointer.h"
+
+namespace UE
+{
+	namespace DatasmithInterchange
+	{
+		class FDatasmithCityEngineMaterialSelector : public FDatasmithReferenceMaterialSelector
+		{
+		public:
+			FDatasmithCityEngineMaterialSelector();
+
+			virtual const TCHAR* GetMaterialPath(EDatasmithReferenceMaterialType MaterialType) const override;
+#if WITH_EDITOR
+			virtual void PostImportProcess(EDatasmithReferenceMaterialType MaterialType, EDatasmithReferenceMaterialQuality MaterialQuality, UMaterialInstanceConstant* MaterialInstance) const override;
+#endif
+		};
+	}
+}
diff --git a/Engine/Plugins/Experimental/Enterprise/DatasmithInterchange/Source/DatasmithInterchange/Private/InterchangeReferenceMaterials/DatasmithReferenceMaterialManager.cpp b/Engine/Plugins/Experimental/Enterprise/DatasmithInterchange/Source/DatasmithInterchange/Private/InterchangeReferenceMaterials/DatasmithReferenceMaterialManager.cpp
new file mode 100644
--- /dev/null
+++ b/Engine/Plugins/Experimental/Enterprise/DatasmithInterchange/Source/DatasmithInterchange/Private/InterchangeReferenceMaterials/DatasmithReferenceMaterialManager.cpp
@@ -0,0 +1,71 @@
+// Copyright Epic Games, Inc. All Rights Reserved.
+
+#include "InterchangeReferenceMaterials/DatasmithReferenceMaterialManager.h"
+
+#include "InterchangeReferenceMaterials/DatasmithReferenceMaterialSelector.h"
+
+namespace UE
+{
+	namespace DatasmithInterchange
+	{
+		TUniquePtr< FDatasmithReferenceMaterialManager > FDatasmithReferenceMaterialManager::Instance;
+
+		void FDatasmithReferenceMaterialManager::Create()
+		{
+			Instance = MakeUnique< FDatasmithReferenceMaterialManager >();
+		}
+
+		void FDatasmithReferenceMaterialManager::Destroy()
+		{
+			Instance.Reset();
+		}
+
+		FDatasmithReferenceMaterialManager& FDatasmithReferenceMaterialManager::Get()
+		{
+			check(Instance.IsValid());
+			return *Instance.Get();
+		}
+
+		FString FDatasmithReferenceMaterialManager::GetHostFromString(const FString& HostString)
+		{
+			if (HostString.Contains(TEXT("CityEngine")))
+			{
+				return TEXT("CityEngine");
+			}
+			else if (HostString.Contains(TEXT("Deltagen")))
+			{
+				return TEXT("Deltagen");
+			}
+			else if (HostString.Contains(TEXT("VRED")))
+			{
+				return TEXT("VRED");
+			}
+			else
+			{
+				return HostString;
+			}
+		}
+
+		void FDatasmithReferenceMaterialManager::RegisterSelector(const TCHAR* Host, TSharedPtr< FDatasmithReferenceMaterialSelector > Selector)
+		{
+			Selectors.FindOrAdd(Host) = Selector;
+		}
+
+		void FDatasmithReferenceMaterialManager::UnregisterSelector(const TCHAR* Host)
+		{
+			Selectors.Remove(Host);
+		}
+
+		const TSharedPtr< FDatasmithReferenceMaterialSelector > FDatasmithReferenceMaterialManager::GetSelector(const TCHAR* Host) const
+		{
+			if (Selectors.Contains(Host))
+			{
+				return Selectors[Host];
+			}
+			else
+			{
+				return MakeShared< FDatasmithReferenceMaterialSelector >();
+			}
+		}
+	}
+}
\ No newline at end of file
diff --git a/Engine/Plugins/Experimental/Enterprise/DatasmithInterchange/Source/DatasmithInterchange/Private/InterchangeReferenceMaterials/DatasmithReferenceMaterialSelector.cpp b/Engine/Plugins/Experimental/Enterprise/DatasmithInterchange/Source/DatasmithInterchange/Private/InterchangeReferenceMaterials/DatasmithReferenceMaterialSelector.cpp
new file mode 100644
--- /dev/null
+++ b/Engine/Plugins/Experimental/Enterprise/DatasmithInterchange/Source/DatasmithInterchange/Private/InterchangeReferenceMaterials/DatasmithReferenceMaterialSelector.cpp
@@ -0,0 +1,3 @@
+// Copyright Epic Games, Inc. All Rights Reserved.
+
+#include "InterchangeReferenceMaterials/DatasmithReferenceMaterialSelector.h"
diff --git a/Engine/Plugins/Experimental/Enterprise/DatasmithInterchange/Source/DatasmithInterchange/Private/InterchangeReferenceMaterials/DatasmithRevitMaterialSelector.cpp b/Engine/Plugins/Experimental/Enterprise/DatasmithInterchange/Source/DatasmithInterchange/Private/InterchangeReferenceMaterials/DatasmithRevitMaterialSelector.cpp
new file mode 100644
--- /dev/null
+++ b/Engine/Plugins/Experimental/Enterprise/DatasmithInterchange/Source/DatasmithInterchange/Private/InterchangeReferenceMaterials/DatasmithRevitMaterialSelector.cpp
@@ -0,0 +1,53 @@
+// Copyright Epic Games, Inc. All Rights Reserved.
+
+#include "DatasmithRevitMaterialSelector.h"
+
+#include "IDatasmithSceneElements.h"
+
+#include "Materials/Material.h"
+#include "Materials/MaterialInstanceConstant.h"
+#include "Misc/PackageName.h"
+//#include "UObject/SoftObjectPath.h"
+//#include "Templates/Casts.h"
+
+namespace UE
+{
+	namespace DatasmithInterchange
+	{
+		namespace Revit
+		{
+			const TCHAR* MaterialPath = TEXT("/DatasmithContent/Materials/RevitReference.RevitReference");
+		}
+
+		FDatasmithRevitMaterialSelector::FDatasmithRevitMaterialSelector()
+		{
+			bIsValid = FPackageName::DoesPackageExist(Revit::MaterialPath);
+		}
+
+		const TCHAR* FDatasmithRevitMaterialSelector::GetMaterialPath(EDatasmithReferenceMaterialType) const
+		{
+			return Revit::MaterialPath;
+		}
+
+#if WITH_EDITOR
+		void FDatasmithRevitMaterialSelector::PostImportProcess(EDatasmithReferenceMaterialType MaterialType, EDatasmithReferenceMaterialQuality, UMaterialInstanceConstant* MaterialInstance) const
+		{
+			if (MaterialInstance)
+			{
+				// Set blend mode to translucent if material requires transparency
+				if (MaterialType == EDatasmithReferenceMaterialType::Transparent)
+				{
+					MaterialInstance->BasePropertyOverrides.bOverride_BlendMode = true;
+					MaterialInstance->BasePropertyOverrides.BlendMode = BLEND_Translucent;
+				}
+				// Set blend mode to masked if material has cutouts
+				else if (MaterialType == EDatasmithReferenceMaterialType::CutOut)
+				{
+					MaterialInstance->BasePropertyOverrides.bOverride_BlendMode = true;
+					MaterialInstance->BasePropertyOverrides.BlendMode = BLEND_Masked;
+				}
+			}
+		}
+#endif
+	}
+}
\ No newline at end of file
diff --git a/Engine/Plugins/Experimental/Enterprise/DatasmithInterchange/Source/DatasmithInterchange/Private/InterchangeReferenceMaterials/DatasmithRevitMaterialSelector.h b/Engine/Plugins/Experimental/Enterprise/DatasmithInterchange/Source/DatasmithInterchange/Private/InterchangeReferenceMaterials/DatasmithRevitMaterialSelector.h
new file mode 100644
--- /dev/null
+++ b/Engine/Plugins/Experimental/Enterprise/DatasmithInterchange/Source/DatasmithInterchange/Private/InterchangeReferenceMaterials/DatasmithRevitMaterialSelector.h
@@ -0,0 +1,25 @@
+// Copyright Epic Games, Inc. All Rights Reserved.
+#pragma once
+
+#include "DatasmithDefinitions.h"
+#include "InterchangeReferenceMaterials/DatasmithReferenceMaterialSelector.h"
+
+#include "Templates/SharedPointer.h"
+
+namespace UE
+{
+	namespace DatasmithInterchange
+	{
+		class FDatasmithRevitMaterialSelector : public FDatasmithReferenceMaterialSelector
+		{
+		public:
+			FDatasmithRevitMaterialSelector();
+			virtual ~FDatasmithRevitMaterialSelector() = default;
+
+			virtual const TCHAR* GetMaterialPath(EDatasmithReferenceMaterialType MaterialType) const override;
+#if WITH_EDITOR
+			virtual void PostImportProcess(EDatasmithReferenceMaterialType MaterialType, EDatasmithReferenceMaterialQuality MaterialQuality, UMaterialInstanceConstant* MaterialInstance) const override;
+#endif
+		};
+	}
+}
diff --git a/Engine/Plugins/Experimental/Enterprise/DatasmithInterchange/Source/DatasmithInterchange/Private/InterchangeReferenceMaterials/DatasmithSketchupMaterialSelector.cpp b/Engine/Plugins/Experimental/Enterprise/DatasmithInterchange/Source/DatasmithInterchange/Private/InterchangeReferenceMaterials/DatasmithSketchupMaterialSelector.cpp
new file mode 100644
--- /dev/null
+++ b/Engine/Plugins/Experimental/Enterprise/DatasmithInterchange/Source/DatasmithInterchange/Private/InterchangeReferenceMaterials/DatasmithSketchupMaterialSelector.cpp
@@ -0,0 +1,42 @@
+// Copyright Epic Games, Inc. All Rights Reserved.
+
+#include "DatasmithSketchupMaterialSelector.h"
+
+#include "Materials/MaterialInstanceConstant.h"
+#include "Misc/PackageName.h"
+
+namespace UE
+{
+	namespace DatasmithInterchange
+	{
+		namespace SketchUp
+		{
+			const TCHAR* MaterialPath = TEXT("/DatasmithContent/Materials/SketchupReference.SketchupReference");
+		}
+
+		FDatasmithSketchUpMaterialSelector::FDatasmithSketchUpMaterialSelector()
+		{
+			bIsValid = FPackageName::DoesPackageExist(SketchUp::MaterialPath);
+		}
+
+		const TCHAR* FDatasmithSketchUpMaterialSelector::GetMaterialPath(EDatasmithReferenceMaterialType MaterialType) const
+		{
+			return SketchUp::MaterialPath;
+		}
+
+#if WITH_EDITOR
+		void FDatasmithSketchUpMaterialSelector::PostImportProcess(EDatasmithReferenceMaterialType MaterialType, EDatasmithReferenceMaterialQuality, UMaterialInstanceConstant* MaterialInstance) const
+		{
+			if (MaterialInstance)
+			{
+				if (MaterialType == EDatasmithReferenceMaterialType::Transparent)
+				{
+					MaterialInstance->BasePropertyOverrides.bOverride_BlendMode = true;
+					MaterialInstance->BasePropertyOverrides.BlendMode = BLEND_Translucent;
+				}
+			}
+
+		}
+#endif
+	}
+}
\ No newline at end of file
diff --git a/Engine/Plugins/Experimental/Enterprise/DatasmithInterchange/Source/DatasmithInterchange/Private/InterchangeReferenceMaterials/DatasmithSketchupMaterialSelector.h b/Engine/Plugins/Experimental/Enterprise/DatasmithInterchange/Source/DatasmithInterchange/Private/InterchangeReferenceMaterials/DatasmithSketchupMaterialSelector.h
new file mode 100644
--- /dev/null
+++ b/Engine/Plugins/Experimental/Enterprise/DatasmithInterchange/Source/DatasmithInterchange/Private/InterchangeReferenceMaterials/DatasmithSketchupMaterialSelector.h
@@ -0,0 +1,23 @@
+// Copyright Epic Games, Inc. All Rights Reserved.
+#pragma once
+
+#include "DatasmithDefinitions.h"
+#include "InterchangeReferenceMaterials/DatasmithReferenceMaterialSelector.h"
+
+namespace UE
+{
+	namespace DatasmithInterchange
+	{
+		class FDatasmithSketchUpMaterialSelector : public FDatasmithReferenceMaterialSelector
+		{
+		public:
+			FDatasmithSketchUpMaterialSelector();
+			virtual ~FDatasmithSketchUpMaterialSelector() = default;
+
+			virtual const TCHAR* GetMaterialPath(EDatasmithReferenceMaterialType MaterialType) const override;
+#if WITH_EDITOR
+			virtual void PostImportProcess(EDatasmithReferenceMaterialType MaterialType, EDatasmithReferenceMaterialQuality MaterialQuality, UMaterialInstanceConstant* MaterialInstance) const override;
+#endif
+		};
+	}
+}
\ No newline at end of file
diff --git a/Engine/Plugins/Experimental/Enterprise/DatasmithInterchange/Source/DatasmithInterchange/Private/InterchangeReferenceMaterials/DatasmithStdMaterialSelector.cpp b/Engine/Plugins/Experimental/Enterprise/DatasmithInterchange/Source/DatasmithInterchange/Private/InterchangeReferenceMaterials/DatasmithStdMaterialSelector.cpp
new file mode 100644
--- /dev/null
+++ b/Engine/Plugins/Experimental/Enterprise/DatasmithInterchange/Source/DatasmithInterchange/Private/InterchangeReferenceMaterials/DatasmithStdMaterialSelector.cpp
@@ -0,0 +1,39 @@
+// Copyright Epic Games, Inc. All Rights Reserved.
+
+#include "DatasmithStdMaterialSelector.h"
+
+#include "Misc/PackageName.h"
+
+namespace UE
+{
+	namespace DatasmithInterchange
+	{
+		namespace Standard
+		{
+			const TCHAR* OpaqueMaterialPath = TEXT("/DatasmithContent/Materials/StdOpaque/M_StdOpaque.M_StdOpaque");
+			const TCHAR* TranslucentMaterialPath = TEXT("/DatasmithContent/Materials/StdTranslucent/M_StdTranslucent.M_StdTranslucent");
+			const TCHAR* EmissiveMaterialPath = TEXT("/DatasmithContent/Materials/StdEmissive/M_StdEmissive.M_StdEmissive");
+		}
+
+		FDatasmithStdMaterialSelector::FDatasmithStdMaterialSelector()
+		{
+			bIsValid = FPackageName::DoesPackageExist(Standard::OpaqueMaterialPath) && FPackageName::DoesPackageExist(Standard::TranslucentMaterialPath) && FPackageName::DoesPackageExist(Standard::EmissiveMaterialPath);
+		}
+
+		const TCHAR* FDatasmithStdMaterialSelector::GetMaterialPath(EDatasmithReferenceMaterialType MaterialType) const
+		{
+			switch (MaterialType)
+			{
+			case EDatasmithReferenceMaterialType::Transparent:
+				return Standard::TranslucentMaterialPath;
+				break;
+			case EDatasmithReferenceMaterialType::Emissive:
+				return Standard::EmissiveMaterialPath;
+				break;
+			default:
+				return Standard::OpaqueMaterialPath;
+				break;
+			}
+		}
+	}
+}
\ No newline at end of file
diff --git a/Engine/Plugins/Experimental/Enterprise/DatasmithInterchange/Source/DatasmithInterchange/Private/InterchangeReferenceMaterials/DatasmithStdMaterialSelector.h b/Engine/Plugins/Experimental/Enterprise/DatasmithInterchange/Source/DatasmithInterchange/Private/InterchangeReferenceMaterials/DatasmithStdMaterialSelector.h
new file mode 100644
--- /dev/null
+++ b/Engine/Plugins/Experimental/Enterprise/DatasmithInterchange/Source/DatasmithInterchange/Private/InterchangeReferenceMaterials/DatasmithStdMaterialSelector.h
@@ -0,0 +1,21 @@
+// Copyright Epic Games, Inc. All Rights Reserved.
+#pragma once
+
+#include "InterchangeReferenceMaterials/DatasmithReferenceMaterialSelector.h"
+
+#include "Templates/SharedPointer.h"
+
+namespace UE
+{
+	namespace DatasmithInterchange
+	{
+		class FDatasmithStdMaterialSelector : public FDatasmithReferenceMaterialSelector
+		{
+		public:
+			FDatasmithStdMaterialSelector();
+			virtual ~FDatasmithStdMaterialSelector() = default;
+
+			virtual const TCHAR* GetMaterialPath(EDatasmithReferenceMaterialType MaterialType) const override;
+		};
+	}
+}
diff --git a/Engine/Plugins/Experimental/Enterprise/DatasmithInterchange/Source/DatasmithInterchange/Public/DatasmithInterchangeModule.h b/Engine/Plugins/Experimental/Enterprise/DatasmithInterchange/Source/DatasmithInterchange/Public/DatasmithInterchangeModule.h
new file mode 100644
--- /dev/null
+++ b/Engine/Plugins/Experimental/Enterprise/DatasmithInterchange/Source/DatasmithInterchange/Public/DatasmithInterchangeModule.h
@@ -0,0 +1,37 @@
+// Copyright Epic Games, Inc. All Rights Reserved.
+
+#pragma once
+
+#include "Modules/ModuleInterface.h"
+#include "Modules/ModuleManager.h"		// For inline LoadModuleChecked()
+
+
+#define DATASMITHINTERCHANGE_MODULE_NAME TEXT("DatasmithInterchange")
+
+/**
+ * The public interface of the DatasmithInterchange module
+ */
+class IDatasmithInterchangeModule : public IModuleInterface
+{
+public:
+
+	/**
+	 * Singleton-like access to IDatasmithImporter
+	 *
+	 * @return Returns DatasmithImporter singleton instance, loading the module on demand if needed
+	 */
+	static inline IDatasmithInterchangeModule& Get()
+	{
+		return FModuleManager::LoadModuleChecked<IDatasmithInterchangeModule>(DATASMITHINTERCHANGE_MODULE_NAME);
+	}
+
+	/**
+	 * Checks to see if this module is loaded and ready.  It is only valid to call Get() if IsAvailable() returns true.
+	 *
+	 * @return True if the module is loaded and ready to use
+	 */
+	static inline bool IsAvailable()
+	{
+		return FModuleManager::Get().IsModuleLoaded(DATASMITHINTERCHANGE_MODULE_NAME);
+	}
+};
\ No newline at end of file
diff --git a/Engine/Plugins/Experimental/Enterprise/DatasmithInterchange/Source/DatasmithInterchange/Public/InterchangeDatasmithAreaLightFactory.h b/Engine/Plugins/Experimental/Enterprise/DatasmithInterchange/Source/DatasmithInterchange/Public/InterchangeDatasmithAreaLightFactory.h
new file mode 100644
--- /dev/null
+++ b/Engine/Plugins/Experimental/Enterprise/DatasmithInterchange/Source/DatasmithInterchange/Public/InterchangeDatasmithAreaLightFactory.h
@@ -0,0 +1,28 @@
+// Copyright Epic Games, Inc. All Rights Reserved.
+
+#pragma once
+
+#include "CoreMinimal.h"
+#include "Scene/InterchangeActorFactory.h"
+#include "UObject/Object.h"
+#include "UObject/ObjectMacros.h"
+
+#include "InterchangeDatasmithAreaLightFactory.generated.h"
+
+
+UCLASS(BlueprintType, Experimental)
+class DATASMITHINTERCHANGE_API UInterchangeDatasmithAreaLightFactory : public UInterchangeActorFactory
+{
+	GENERATED_BODY()
+
+public:
+	//////////////////////////////////////////////////////////////////////////
+	// Interchange factory base interface begin
+
+	virtual UClass* GetFactoryClass() const override;
+
+	virtual UObject* CreateSceneObject(const FCreateSceneObjectsParams& Arguments) override;
+
+	// Interchange factory base interface end
+	//////////////////////////////////////////////////////////////////////////
+};
\ No newline at end of file
diff --git a/Engine/Plugins/Experimental/Enterprise/DatasmithInterchange/Source/DatasmithInterchange/Public/InterchangeDatasmithAreaLightFactoryNode.h b/Engine/Plugins/Experimental/Enterprise/DatasmithInterchange/Source/DatasmithInterchange/Public/InterchangeDatasmithAreaLightFactoryNode.h
new file mode 100644
--- /dev/null
+++ b/Engine/Plugins/Experimental/Enterprise/DatasmithInterchange/Source/DatasmithInterchange/Public/InterchangeDatasmithAreaLightFactoryNode.h
@@ -0,0 +1,147 @@
+// Copyright Epic Games, Inc. All Rights Reserved.
+
+#pragma once
+
+#include "CoreMinimal.h"
+#include "InterchangeActorFactoryNode.h"
+#include "Nodes/InterchangeFactoryBaseNode.h"
+
+#include "DatasmithAreaLightActor.h"
+
+#include "InterchangeDatasmithAreaLightFactoryNode.generated.h"
+
+UCLASS(BlueprintType, Experimental)
+class DATASMITHINTERCHANGE_API UInterchangeDatasmithAreaLightFactoryNode : public UInterchangeActorFactoryNode
+{
+	GENERATED_BODY()
+
+public:
+
+	/**
+	 * Return the node type name of the class, we use this when reporting error
+	 */
+	virtual FString GetTypeName() const override;
+
+	virtual UClass* GetObjectClass() const override;
+
+
+	UFUNCTION(BlueprintCallable, Category = "Interchange | Node | Datasmith")
+	bool GetCustomLightType(EDatasmithAreaLightActorType& AttributeValue) const;
+
+	UFUNCTION(BlueprintCallable, Category = "Interchange | Node | Datasmith")
+	bool SetCustomLightType(const EDatasmithAreaLightActorType& AttributeValue, bool bAddApplyDelegate = true);
+
+	UFUNCTION(BlueprintCallable, Category = "Interchange | Node | Datasmith")
+	bool GetCustomLightShape(EDatasmithAreaLightActorShape& AttributeValue) const;
+
+	UFUNCTION(BlueprintCallable, Category = "Interchange | Node | Datasmith")
+	bool SetCustomLightShape(const EDatasmithAreaLightActorShape& AttributeValue, bool bAddApplyDelegate = true);
+
+	UFUNCTION(BlueprintCallable, Category = "Interchange | Node | Datasmith")
+	bool GetCustomDimensions(FVector2D& AttributeValue) const;
+
+	UFUNCTION(BlueprintCallable, Category = "Interchange | Node | Datasmith")
+	bool SetCustomDimensions(const FVector2D& AttributeValue, bool bAddApplyDelegate = true);
+
+	UFUNCTION(BlueprintCallable, Category = "Interchange | Node | Datasmith")
+	bool GetCustomIntensity(float& AttributeValue) const;
+
+	UFUNCTION(BlueprintCallable, Category = "Interchange | Node | Datasmith")
+	bool SetCustomIntensity(const float& AttributeValue, bool bAddApplyDelegate = true);
+
+	UFUNCTION(BlueprintCallable, Category = "Interchange | Node | Datasmith")
+	bool GetCustomIntensityUnits(ELightUnits& AttributeValue) const;
+
+	UFUNCTION(BlueprintCallable, Category = "Interchange | Node | Datasmith")
+	bool SetCustomIntensityUnits(const ELightUnits& AttributeValue, bool bAddApplyDelegate = true);
+
+	UFUNCTION(BlueprintCallable, Category = "Interchange | Node | Datasmith")
+	bool GetCustomColor(FLinearColor& AttributeValue) const;
+
+	UFUNCTION(BlueprintCallable, Category = "Interchange | Node | Datasmith")
+	bool SetCustomColor(const FLinearColor& AttributeValue, bool bAddApplyDelegate = true);
+
+	UFUNCTION(BlueprintCallable, Category = "Interchange | Node | Datasmith")
+	bool GetCustomTemperature(float& AttributeValue) const;
+
+	UFUNCTION(BlueprintCallable, Category = "Interchange | Node | Datasmith")
+	bool SetCustomTemperature(const float& AttributeValue, bool bAddApplyDelegate = true);
+
+	/**
+	 * Get the NodeUid of the imported IES Texture.
+	 */
+	//UFUNCTION(BlueprintCallable, Category = "Interchange | Node | Datasmith")
+	bool GetCustomIESTexture(FString& AttributeValue) const;
+
+	/**
+	 * Set the NodeUid of the imported IES Texture.
+	 */
+	//UFUNCTION(BlueprintCallable, Category = "Interchange | Node | Datasmith")
+	bool SetCustomIESTexture(const FString& AttributeValue);
+
+	UFUNCTION(BlueprintCallable, Category = "Interchange | Node | Datasmith")
+	bool GetCustomUseIESBrightness(bool& AttributeValue) const;
+
+	UFUNCTION(BlueprintCallable, Category = "Interchange | Node | Datasmith")
+	bool SetCustomUseIESBrightness(const bool& AttributeValue, bool bAddApplyDelegate = true);
+
+	UFUNCTION(BlueprintCallable, Category = "Interchange | Node | Datasmith")
+	bool GetCustomIESBrightnessScale(float& AttributeValue) const;
+	
+	UFUNCTION(BlueprintCallable, Category = "Interchange | Node | Datasmith")
+	bool SetCustomIESBrightnessScale(const float& AttributeValue, bool bAddApplyDelegate = true);
+
+	UFUNCTION(BlueprintCallable, Category = "Interchange | Node | Datasmith")
+	bool GetCustomRotation(FRotator& AttributeValue) const;
+	
+	UFUNCTION(BlueprintCallable, Category = "Interchange | Node | Datasmith")
+	bool SetCustomRotation(const FRotator& AttributeValue, bool bAddApplyDelegate = true);
+
+	UFUNCTION(BlueprintCallable, Category = "Interchange | Node | Datasmith")
+	bool GetCustomSourceRadius(float& AttributeValue) const;
+	
+	UFUNCTION(BlueprintCallable, Category = "Interchange | Node | Datasmith")
+	bool SetCustomSourceRadius(const float& AttributeValue, bool bAddApplyDelegate = true);
+
+	UFUNCTION(BlueprintCallable, Category = "Interchange | Node | Datasmith")
+	bool GetCustomSourceLength(float& AttributeValue) const;
+	
+	UFUNCTION(BlueprintCallable, Category = "Interchange | Node | Datasmith")
+	bool SetCustomSourceLength(const float& AttributeValue, bool bAddApplyDelegate = true);
+
+	UFUNCTION(BlueprintCallable, Category = "Interchange | Node | Datasmith")
+	bool GetCustomAttenuationRadius(float& AttributeValue) const;
+	
+	UFUNCTION(BlueprintCallable, Category = "Interchange | Node | Datasmith")
+	bool SetCustomAttenuationRadius(const float& AttributeValue, bool bAddApplyDelegate = true);
+
+	UFUNCTION(BlueprintCallable, Category = "Interchange | Node | Datasmith")
+	bool GetCustomSpotlightInnerAngle(float& AttributeValue) const;
+	
+	UFUNCTION(BlueprintCallable, Category = "Interchange | Node | Datasmith")
+	bool SetCustomSpotlightInnerAngle(const float& AttributeValue, bool bAddApplyDelegate = true);
+
+	UFUNCTION(BlueprintCallable, Category = "Interchange | Node | Datasmith")
+	bool GetCustomSpotlightOuterAngle(float& AttributeValue) const;
+	
+	UFUNCTION(BlueprintCallable, Category = "Interchange | Node | Datasmith")
+	bool SetCustomSpotlightOuterAngle(const float& AttributeValue, bool bAddApplyDelegate = true);
+
+private:
+	static  UE::Interchange::FAttributeKey Macro_CustomLightTypeKey;
+	static  UE::Interchange::FAttributeKey Macro_CustomLightShapeKey;
+	static  UE::Interchange::FAttributeKey Macro_CustomDimensionsKey;
+	static  UE::Interchange::FAttributeKey Macro_CustomIntensityKey;
+	static  UE::Interchange::FAttributeKey Macro_CustomIntensityUnitsKey;
+	static  UE::Interchange::FAttributeKey Macro_CustomColorKey;
+	static  UE::Interchange::FAttributeKey Macro_CustomTemperatureKey;
+	static  UE::Interchange::FAttributeKey Macro_CustomIESTextureKey;
+	static  UE::Interchange::FAttributeKey Macro_CustomUseIESBrightnessKey;
+	static  UE::Interchange::FAttributeKey Macro_CustomIESBrightnessScaleKey;
+	static  UE::Interchange::FAttributeKey Macro_CustomRotationKey;
+	static  UE::Interchange::FAttributeKey Macro_CustomSourceRadiusKey;
+	static  UE::Interchange::FAttributeKey Macro_CustomSourceLengthKey;
+	static  UE::Interchange::FAttributeKey Macro_CustomAttenuationRadiusKey;
+	static  UE::Interchange::FAttributeKey Macro_CustomSpotlightInnerAngleKey;
+	static  UE::Interchange::FAttributeKey Macro_CustomSpotlightOuterAngleKey;
+};
\ No newline at end of file
diff --git a/Engine/Plugins/Experimental/Enterprise/DatasmithInterchange/Source/DatasmithInterchange/Public/InterchangeDatasmithAreaLightNode.h b/Engine/Plugins/Experimental/Enterprise/DatasmithInterchange/Source/DatasmithInterchange/Public/InterchangeDatasmithAreaLightNode.h
new file mode 100644
--- /dev/null
+++ b/Engine/Plugins/Experimental/Enterprise/DatasmithInterchange/Source/DatasmithInterchange/Public/InterchangeDatasmithAreaLightNode.h
@@ -0,0 +1,125 @@
+// Copyright Epic Games, Inc. All Rights Reserved.
+
+#pragma once
+
+#include "CoreMinimal.h"
+#include "Types/AttributeStorage.h"
+#include "InterchangeLightNode.h"
+#include "DatasmithAreaLightActor.h"
+
+#include "InterchangeDatasmithAreaLightNode.generated.h"
+
+namespace UE::DatasmithImporter
+{
+	class FExternalSource;
+}
+
+class IDatasmithScene;
+
+UCLASS(BlueprintType, Experimental)
+class DATASMITHINTERCHANGE_API UInterchangeDatasmithAreaLightNode : public UInterchangeLightNode
+{
+	GENERATED_BODY()
+
+public:
+	/**
+	 * Return the node type name of the class, we use this when reporting error
+	 */
+	virtual FString GetTypeName() const override
+	{
+		const FString TypeName = TEXT("DatasmithAreaLightNode");
+		return TypeName;
+	}
+
+	UFUNCTION(BlueprintCallable, Category = "Interchange | Node | Datasmith")
+	bool GetCustomLightType(EDatasmithAreaLightActorType& AttributeValue) const;
+
+	UFUNCTION(BlueprintCallable, Category = "Interchange | Node | Datasmith")
+	bool SetCustomLightType(const EDatasmithAreaLightActorType& AttributeValue, bool bAddApplyDelegate = true);
+
+	UFUNCTION(BlueprintCallable, Category = "Interchange | Node | Datasmith")
+	bool GetCustomLightShape(EDatasmithAreaLightActorShape& AttributeValue) const;
+
+	UFUNCTION(BlueprintCallable, Category = "Interchange | Node | Datasmith")
+	bool SetCustomLightShape(const EDatasmithAreaLightActorShape& AttributeValue, bool bAddApplyDelegate = true);
+
+	UFUNCTION(BlueprintCallable, Category = "Interchange | Node | Datasmith")
+	bool GetCustomDimensions(FVector2D& AttributeValue) const;
+
+	UFUNCTION(BlueprintCallable, Category = "Interchange | Node | Datasmith")
+	bool SetCustomDimensions(const FVector2D& AttributeValue, bool bAddApplyDelegate = true);
+
+	UFUNCTION(BlueprintCallable, Category = "Interchange | Node | Datasmith")
+	bool GetCustomColor(FLinearColor& AttributeValue) const;
+
+	UFUNCTION(BlueprintCallable, Category = "Interchange | Node | Datasmith")
+	bool SetCustomColor(const FLinearColor& AttributeValue, bool bAddApplyDelegate = true);
+
+	/**
+	 * Get the NodeUid of the imported IES Texture.
+	 */
+	//UFUNCTION(BlueprintCallable, Category = "Interchange | Node | Datasmith")
+	bool GetCustomIESTexture(FString& AttributeValue) const;
+
+	/**
+	 * Set the NodeUid of the imported IES Texture.
+	 */
+	//UFUNCTION(BlueprintCallable, Category = "Interchange | Node | Datasmith")
+	bool SetCustomIESTexture(const FString& AttributeValue);
+
+	UFUNCTION(BlueprintCallable, Category = "Interchange | Node | Datasmith")
+	bool GetCustomUseIESBrightness(bool& AttributeValue) const;
+
+	UFUNCTION(BlueprintCallable, Category = "Interchange | Node | Datasmith")
+	bool SetCustomUseIESBrightness(const bool& AttributeValue, bool bAddApplyDelegate = true);
+
+	UFUNCTION(BlueprintCallable, Category = "Interchange | Node | Datasmith")
+	bool GetCustomIESBrightnessScale(float& AttributeValue) const;
+
+	UFUNCTION(BlueprintCallable, Category = "Interchange | Node | Datasmith")
+	bool SetCustomIESBrightnessScale(const float& AttributeValue, bool bAddApplyDelegate = true);
+
+	UFUNCTION(BlueprintCallable, Category = "Interchange | Node | Datasmith")
+	bool GetCustomRotation(FRotator& AttributeValue) const;
+
+	UFUNCTION(BlueprintCallable, Category = "Interchange | Node | Datasmith")
+	bool SetCustomRotation(const FRotator& AttributeValue, bool bAddApplyDelegate = true);
+
+	UFUNCTION(BlueprintCallable, Category = "Interchange | Node | Datasmith")
+	bool GetCustomSourceRadius(float& AttributeValue) const;
+
+	UFUNCTION(BlueprintCallable, Category = "Interchange | Node | Datasmith")
+	bool SetCustomSourceRadius(const float& AttributeValue, bool bAddApplyDelegate = true);
+
+	UFUNCTION(BlueprintCallable, Category = "Interchange | Node | Datasmith")
+	bool GetCustomSourceLength(float& AttributeValue) const;
+
+	UFUNCTION(BlueprintCallable, Category = "Interchange | Node | Datasmith")
+	bool SetCustomSourceLength(const float& AttributeValue, bool bAddApplyDelegate = true);
+
+	UFUNCTION(BlueprintCallable, Category = "Interchange | Node | Datasmith")
+	bool GetCustomSpotlightInnerAngle(float& AttributeValue) const;
+
+	UFUNCTION(BlueprintCallable, Category = "Interchange | Node | Datasmith")
+	bool SetCustomSpotlightInnerAngle(const float& AttributeValue, bool bAddApplyDelegate = true);
+
+	UFUNCTION(BlueprintCallable, Category = "Interchange | Node | Datasmith")
+	bool GetCustomSpotlightOuterAngle(float& AttributeValue) const;
+
+	UFUNCTION(BlueprintCallable, Category = "Interchange | Node | Datasmith")
+	bool SetCustomSpotlightOuterAngle(const float& AttributeValue, bool bAddApplyDelegate = true);
+
+private:
+	static  UE::Interchange::FAttributeKey Macro_CustomLightTypeKey;
+	static  UE::Interchange::FAttributeKey Macro_CustomLightShapeKey;
+	static  UE::Interchange::FAttributeKey Macro_CustomDimensionsKey;
+	static  UE::Interchange::FAttributeKey Macro_CustomColorKey;
+	static  UE::Interchange::FAttributeKey Macro_CustomIESTextureKey;
+	static  UE::Interchange::FAttributeKey Macro_CustomUseIESBrightnessKey;
+	static  UE::Interchange::FAttributeKey Macro_CustomIESBrightnessScaleKey;
+	static  UE::Interchange::FAttributeKey Macro_CustomRotationKey;
+	static  UE::Interchange::FAttributeKey Macro_CustomSourceRadiusKey;
+	static  UE::Interchange::FAttributeKey Macro_CustomSourceLengthKey;
+	static  UE::Interchange::FAttributeKey Macro_CustomSpotlightInnerAngleKey;
+	static  UE::Interchange::FAttributeKey Macro_CustomSpotlightOuterAngleKey;
+};
\ No newline at end of file
diff --git a/Engine/Plugins/Experimental/Enterprise/DatasmithInterchange/Source/DatasmithInterchange/Public/InterchangeDatasmithPipeline.h b/Engine/Plugins/Experimental/Enterprise/DatasmithInterchange/Source/DatasmithInterchange/Public/InterchangeDatasmithPipeline.h
new file mode 100644
--- /dev/null
+++ b/Engine/Plugins/Experimental/Enterprise/DatasmithInterchange/Source/DatasmithInterchange/Public/InterchangeDatasmithPipeline.h
@@ -0,0 +1,68 @@
+// Copyright Epic Games, Inc. All Rights Reserved.
+
+#pragma once
+
+#include "CoreMinimal.h"
+#include "InterchangePipelineBase.h"
+#include "UObject/Object.h"
+#include "UObject/ObjectMacros.h"
+
+#include "InterchangeDatasmithPipeline.generated.h"
+
+class UDatasmithScene;
+class UMaterialInterface;
+class UInterchangeGenericAnimationPipeline;
+class UInterchangeGenericCommonMeshesProperties;
+class UInterchangeGenericCommonSkeletalMeshesAndAnimationsProperties;
+class UInterchangeDatasmithLevelPipeline;
+class UInterchangeDatasmithMaterialPipeline;
+class UInterchangeDatasmithStaticMeshPipeline;
+class UInterchangeDatasmithSceneNode;
+class UInterchangeDatasmithTexturePipeline;
+class UTexture;
+struct FDatasmithImportContext;
+
+UCLASS(BlueprintType, Experimental)
+class DATASMITHINTERCHANGE_API UInterchangeDatasmithPipeline : public UInterchangePipelineBase
+{
+	GENERATED_BODY()
+
+	UInterchangeDatasmithPipeline();
+
+	UPROPERTY(VisibleAnywhere, Instanced, Category = "Textures")
+	TObjectPtr<UInterchangeDatasmithTexturePipeline> TexturePipeline;
+
+	UPROPERTY(VisibleAnywhere, Instanced, Category = "Materials")
+	TObjectPtr<UInterchangeDatasmithMaterialPipeline> MaterialPipeline;
+
+	UPROPERTY(VisibleAnywhere, Instanced, Category = "Mesh")
+	TObjectPtr<UInterchangeDatasmithStaticMeshPipeline> MeshPipeline;
+
+	UPROPERTY(VisibleAnywhere, Instanced, Category = "Level")
+	TObjectPtr<UInterchangeDatasmithLevelPipeline> LevelPipeline;
+
+	UPROPERTY(VisibleAnywhere, Instanced, Category = "Animation")
+	TObjectPtr<UInterchangeGenericAnimationPipeline> AnimationPipeline;
+
+	UPROPERTY(VisibleAnywhere, Instanced, Category = "Common Meshes")
+	TObjectPtr<UInterchangeGenericCommonMeshesProperties> CommonMeshesProperties;
+
+	UPROPERTY(VisibleAnywhere, Instanced, Category = "Common Skeletal Meshes and Animations")
+	TObjectPtr<UInterchangeGenericCommonSkeletalMeshesAndAnimationsProperties> CommonSkeletalMeshesAndAnimationsProperties;
+
+protected:
+	virtual void ExecutePreImportPipeline(UInterchangeBaseNodeContainer* InBaseNodeContainer, const TArray<UInterchangeSourceData*>& SourceDatas) override;
+	virtual void ExecutePostImportPipeline(const UInterchangeBaseNodeContainer* InBaseNodeContainer, const FString& NodeKey, UObject* CreatedAsset, bool bIsAReimport) override;
+
+	virtual bool CanExecuteOnAnyThread(EInterchangePipelineTask PipelineTask) override
+	{
+		// This pipeline creates UObjects and assets. Not safe to execute outside of main thread.
+		return false;
+	}
+
+private:
+	UInterchangeBaseNodeContainer* BaseNodeContainer = nullptr;
+
+	// Fill up the UDatasmithScene with all the data its needs for DnD
+	void PostImportDatasmithSceneAsset(UDatasmithScene& DatasmithSceneAsset);
+};
diff --git a/Engine/Plugins/Experimental/Enterprise/DatasmithInterchange/Source/DatasmithInterchange/Public/InterchangeDatasmithSceneFactory.h b/Engine/Plugins/Experimental/Enterprise/DatasmithInterchange/Source/DatasmithInterchange/Public/InterchangeDatasmithSceneFactory.h
new file mode 100644
--- /dev/null
+++ b/Engine/Plugins/Experimental/Enterprise/DatasmithInterchange/Source/DatasmithInterchange/Public/InterchangeDatasmithSceneFactory.h
@@ -0,0 +1,37 @@
+// Copyright Epic Games, Inc. All Rights Reserved.
+
+#pragma once
+
+#include "CoreMinimal.h"
+//#include "EditorFramework/AssetImportData.h"
+#include "InterchangeFactoryBase.h"
+#include "UObject/Object.h"
+#include "UObject/ObjectMacros.h"
+
+#include "InterchangeDatasmithSceneFactory.generated.h"
+
+
+UCLASS(BlueprintType, Experimental)
+class DATASMITHINTERCHANGE_API UInterchangeDatasmithSceneFactory : public UInterchangeFactoryBase
+{
+	GENERATED_BODY()
+public:
+
+	//////////////////////////////////////////////////////////////////////////
+	// Interchange factory base interface begin
+
+	virtual UClass* GetFactoryClass() const override;
+	virtual UObject* CreateEmptyAsset(const FCreateAssetParams& Arguments) override;
+	virtual UObject* CreateAsset(const FCreateAssetParams& Arguments) override;
+	virtual void PreImportPreCompletedCallback(const FImportPreCompletedCallbackParams& Arguments) override;
+	// Interchange factory base interface end
+	//////////////////////////////////////////////////////////////////////////
+
+private:
+	//UE::Interchange::Private::InterchangeTextureFactory::FProcessedPayload ProcessedPayload;
+
+#if WITH_EDITORONLY_DATA
+	// When importing a UDIM the data for the source files will be stored here
+	//TArray<FAssetImportInfo::FSourceFile> SourceFiles;
+#endif // WITH_EDITORONLY_DATA
+};
\ No newline at end of file
diff --git a/Engine/Plugins/Experimental/Enterprise/DatasmithInterchange/Source/DatasmithInterchange/Public/InterchangeDatasmithSceneFactoryNode.h b/Engine/Plugins/Experimental/Enterprise/DatasmithInterchange/Source/DatasmithInterchange/Public/InterchangeDatasmithSceneFactoryNode.h
new file mode 100644
--- /dev/null
+++ b/Engine/Plugins/Experimental/Enterprise/DatasmithInterchange/Source/DatasmithInterchange/Public/InterchangeDatasmithSceneFactoryNode.h
@@ -0,0 +1,30 @@
+// Copyright Epic Games, Inc. All Rights Reserved.
+
+#pragma once
+
+#include "CoreMinimal.h"
+#include "Nodes/InterchangeFactoryBaseNode.h"
+
+#include "InterchangeDatasmithSceneFactoryNode.generated.h"
+
+UCLASS(BlueprintType, Experimental)
+class DATASMITHINTERCHANGE_API UInterchangeDatasmithSceneFactoryNode : public UInterchangeFactoryBaseNode
+{
+	GENERATED_BODY()
+
+public:
+	void InitializeDatasmithFactorySceneNode(const FString& UniqueID, const FString& DisplayLabel, const FString& InAssetClass)
+	{
+		InitializeNode(UniqueID, DisplayLabel, EInterchangeNodeContainerType::FactoryData);
+
+		FString OperationName = GetTypeName() + TEXT(".SetAssetClassName");
+		InterchangePrivateNodeBase::SetCustomAttribute<FString>(*Attributes, UE::Interchange::FBaseNodeStaticData::ClassTypeAttributeKey(), OperationName, InAssetClass);
+	}
+
+	/**
+	 * Return the node type name of the class, we use this when reporting error
+	 */
+	virtual FString GetTypeName() const override;
+
+	virtual UClass* GetObjectClass() const override;
+};
\ No newline at end of file
diff --git a/Engine/Plugins/Experimental/Enterprise/DatasmithInterchange/Source/DatasmithInterchange/Public/InterchangeDatasmithSceneNode.h b/Engine/Plugins/Experimental/Enterprise/DatasmithInterchange/Source/DatasmithInterchange/Public/InterchangeDatasmithSceneNode.h
new file mode 100644
--- /dev/null
+++ b/Engine/Plugins/Experimental/Enterprise/DatasmithInterchange/Source/DatasmithInterchange/Public/InterchangeDatasmithSceneNode.h
@@ -0,0 +1,36 @@
+// Copyright Epic Games, Inc. All Rights Reserved.
+
+#pragma once
+
+#include "CoreMinimal.h"
+#include "Nodes/InterchangeBaseNode.h"
+
+#include "InterchangeDatasmithSceneNode.generated.h"
+
+namespace UE::DatasmithImporter
+{
+	class FExternalSource;
+}
+
+class IDatasmithScene;
+
+UCLASS(BlueprintType, Experimental)
+class DATASMITHINTERCHANGE_API UInterchangeDatasmithSceneNode : public UInterchangeBaseNode
+{
+	GENERATED_BODY()
+
+public:
+	/**
+	 * Return the node type name of the class, we use this when reporting error
+	 */
+	virtual FString GetTypeName() const override
+	{
+		const FString TypeName = TEXT("DatasmithSceneNode");
+		return TypeName;
+	}
+
+	// Used by the InterchangeDatasmithPipeline while we are in the process of transferring the Datasmith Import logic over Interchange.
+	// It should be removed once we no longer need the DatasmithImportContext.
+	TSharedPtr<UE::DatasmithImporter::FExternalSource> ExternalSource;
+	TSharedPtr<IDatasmithScene> DatasmithScene;
+};
\ No newline at end of file
diff --git a/Engine/Plugins/Experimental/Enterprise/DatasmithInterchange/Source/DatasmithInterchange/Public/InterchangeDatasmithTranslator.h b/Engine/Plugins/Experimental/Enterprise/DatasmithInterchange/Source/DatasmithInterchange/Public/InterchangeDatasmithTranslator.h
new file mode 100644
--- /dev/null
+++ b/Engine/Plugins/Experimental/Enterprise/DatasmithInterchange/Source/DatasmithInterchange/Public/InterchangeDatasmithTranslator.h
@@ -0,0 +1,147 @@
+// Copyright Epic Games, Inc. All Rights Reserved.
+
+#pragma once
+
+#include "CoreMinimal.h"
+
+#include "InterchangeDatasmithUtils.h"
+
+#include "Animation/InterchangeAnimationPayloadInterface.h"
+#include "InterchangeTranslatorBase.h"
+#include "Mesh/InterchangeStaticMeshPayloadInterface.h"
+#include "Texture/InterchangeTexturePayloadData.h"
+#include "Texture/InterchangeTexturePayloadInterface.h"
+#include "Mesh/InterchangeStaticMeshPayloadInterface.h"
+#include "Scene/InterchangeVariantSetPayloadInterface.h"
+
+#include "Async/Async.h"
+#include "ExternalSource.h"
+#include "UObject/GCObjectScopeGuard.h"
+#include "UObject/Object.h"
+#include "UObject/ObjectMacros.h"
+
+#include "InterchangeDatasmithTranslator.generated.h"
+
+class IDatasmithActorElement;
+class IDatasmithBaseAnimationElement;
+class IDatasmithCameraActorElement;
+class IDatasmithLightActorElement;
+class IDatasmithScene;
+class IDatasmithTransformAnimationElement;
+class UInterchangeCameraNode;
+class UInterchangeBaseLightNode;
+class UInterchangeSceneNode;
+
+namespace UE::Interchange
+{
+	struct FImportImage;
+	struct FStaticMeshPayloadData;
+	struct FAnimationTransformPayloadData;
+	struct FAnimationBakeTransformPayloadData;
+}
+
+namespace UE::DatasmithImporter
+{
+	class FExternalSource;
+}
+
+UCLASS(BlueprintType, Experimental)
+class UInterchangeDatasmithTranslator : public UInterchangeTranslatorBase
+	, public IInterchangeTexturePayloadInterface
+	, public IInterchangeStaticMeshPayloadInterface
+	, public IInterchangeAnimationPayloadInterface
+	, public IInterchangeVariantSetPayloadInterface
+{
+	GENERATED_BODY()
+
+public:
+
+	virtual bool CanImportSourceData(const UInterchangeSourceData* InSourceData) const override;
+
+	virtual bool Translate(UInterchangeBaseNodeContainer& BaseNodeContainer) const override;
+
+	virtual TArray<FString> GetSupportedFormats() const override { return TArray<FString>(); }
+	virtual EInterchangeTranslatorAssetType GetSupportedAssetTypes() const override { return EInterchangeTranslatorAssetType::Materials | EInterchangeTranslatorAssetType::Meshes;}
+
+	virtual void ReleaseSource() override
+	{
+		return;
+	}
+
+	virtual void ImportFinish() override;
+
+	/* IInterchangeTexturePayloadInterface Begin */
+	virtual TOptional<UE::Interchange::FImportImage> GetTexturePayloadData(const UInterchangeSourceData* InPayloadSourceData, const FString& PayloadKey) const override;
+	/* IInterchangeTexturePayloadInterface End */
+
+	/* IInterchangeStaticMeshPayloadInterface Begin */
+	virtual TFuture<TOptional<UE::Interchange::FStaticMeshPayloadData>> GetStaticMeshPayloadData(const FString& PayloadKey) const override;
+	/* IInterchangeStaticMeshPayloadInterface End */
+
+	/* IInterchangeAnimationPayloadInterface Begin */
+	virtual TFuture<TOptional<UE::Interchange::FAnimationCurvePayloadData>> GetAnimationCurvePayloadData(const FString& PayLoadKey) const override;
+	virtual TFuture<TOptional<UE::Interchange::FAnimationStepCurvePayloadData>> GetAnimationStepCurvePayloadData(const FString& PayLoadKey) const override;
+	virtual TFuture<TOptional<UE::Interchange::FAnimationBakeTransformPayloadData>> GetAnimationBakeTransformPayloadData(const FString& PayLoadKey, const double BakeFrequency, const double RangeStartSecond, const double RangeStopSecond) const override;
+	/* IInterchangeAnimationPayloadInterface End */
+
+	/* IInterchangeVariantSetPayloadInterface Begin */
+	virtual TFuture<TOptional<UE::Interchange::FVariantSetPayloadData>> GetVariantSetPayloadData(const FString& PayloadKey) const override;
+	/* IInterchangeVariantSetPayloadInterface End */
+
+private:
+
+	void HandleDatasmithActor(UInterchangeBaseNodeContainer& BaseNodeContainer, const TSharedRef<IDatasmithActorElement>& ActorElement, const UInterchangeSceneNode* ParentNode) const;
+
+	UInterchangeCameraNode* AddCameraNode(UInterchangeBaseNodeContainer& BaseNodeContainer, const TSharedRef<IDatasmithCameraActorElement>& CameraActor) const;
+
+	UInterchangeBaseLightNode* AddLightNode(UInterchangeBaseNodeContainer& BaseNodeContainer, const TSharedRef<IDatasmithLightActorElement>& LightActor) const;
+
+	TSharedPtr<UE::DatasmithImporter::FExternalSource> LoadedExternalSource;
+
+	template<typename T>
+	TFuture<TOptional<T>> GetAnimationCurvePayloadData(const FString& PayLoadKey) const
+	{
+		using namespace UE::DatasmithInterchange::AnimUtils;
+
+		TPromise<TOptional<T>> EmptyPromise;
+		EmptyPromise.SetValue(TOptional<T>());
+
+		if (!LoadedExternalSource || !LoadedExternalSource->GetDatasmithScene())
+		{
+			return EmptyPromise.GetFuture();
+		}
+
+		TSharedPtr<IDatasmithBaseAnimationElement> AnimationElement;
+		float FrameRate = 0.f;
+		if (FAnimationPayloadDesc* PayloadDescPtr = AnimationPayLoadMapping.Find(PayLoadKey))
+		{
+			AnimationElement = PayloadDescPtr->Value;
+			if (!ensure(AnimationElement))
+			{
+				// #ueent_logwarning:
+				return EmptyPromise.GetFuture();
+			}
+
+			FrameRate = PayloadDescPtr->Key;
+		}
+
+		return Async(EAsyncExecution::TaskGraph, [this, AnimationElement = MoveTemp(AnimationElement), FrameRate]
+			{
+				T TransformPayloadData;
+				TOptional<T> Result;
+
+				if (GetAnimationPayloadData(*AnimationElement, FrameRate, TransformPayloadData))
+				{
+					Result.Emplace(MoveTemp(TransformPayloadData));
+				}
+
+				return Result;
+			}
+		);
+	}
+
+	mutable uint64 StartTime = 0;
+	mutable FString FileName;
+
+	mutable TMap<FString, UE::DatasmithInterchange::AnimUtils::FAnimationPayloadDesc> AnimationPayLoadMapping;
+};
\ No newline at end of file
diff --git a/Engine/Plugins/Experimental/Enterprise/DatasmithInterchange/Source/DatasmithInterchange/Public/InterchangeReferenceMaterials/DatasmithReferenceMaterialManager.h b/Engine/Plugins/Experimental/Enterprise/DatasmithInterchange/Source/DatasmithInterchange/Public/InterchangeReferenceMaterials/DatasmithReferenceMaterialManager.h
new file mode 100644
--- /dev/null
+++ b/Engine/Plugins/Experimental/Enterprise/DatasmithInterchange/Source/DatasmithInterchange/Public/InterchangeReferenceMaterials/DatasmithReferenceMaterialManager.h
@@ -0,0 +1,34 @@
+// Copyright Epic Games, Inc. All Rights Reserved.
+#pragma once
+
+#include "Containers/ContainersFwd.h"
+#include "Templates/SharedPointer.h"
+#include "Templates/UniquePtr.h"
+
+namespace UE
+{
+	namespace DatasmithInterchange
+	{
+		class FDatasmithReferenceMaterialSelector;
+
+		class DATASMITHINTERCHANGE_API FDatasmithReferenceMaterialManager
+		{
+		public:
+			static void Create();
+			static void Destroy();
+			static FDatasmithReferenceMaterialManager& Get();
+
+			FString GetHostFromString(const FString& HostString);
+
+			void RegisterSelector(const TCHAR* Host, TSharedPtr< FDatasmithReferenceMaterialSelector > Selector);
+			void UnregisterSelector(const TCHAR* Host);
+
+			const TSharedPtr< FDatasmithReferenceMaterialSelector > GetSelector(const TCHAR* Host) const;
+
+		private:
+			static TUniquePtr< FDatasmithReferenceMaterialManager > Instance;
+
+			TMap< FString, TSharedPtr< FDatasmithReferenceMaterialSelector > > Selectors;
+		};
+	}
+}
diff --git a/Engine/Plugins/Experimental/Enterprise/DatasmithInterchange/Source/DatasmithInterchange/Public/InterchangeReferenceMaterials/DatasmithReferenceMaterialSelector.h b/Engine/Plugins/Experimental/Enterprise/DatasmithInterchange/Source/DatasmithInterchange/Public/InterchangeReferenceMaterials/DatasmithReferenceMaterialSelector.h
new file mode 100644
--- /dev/null
+++ b/Engine/Plugins/Experimental/Enterprise/DatasmithInterchange/Source/DatasmithInterchange/Public/InterchangeReferenceMaterials/DatasmithReferenceMaterialSelector.h
@@ -0,0 +1,30 @@
+// Copyright Epic Games, Inc. All Rights Reserved.
+#pragma once
+
+#include "DatasmithDefinitions.h"
+
+class UMaterialInstanceConstant;
+
+namespace UE
+{
+	namespace DatasmithInterchange
+	{
+		class FDatasmithReferenceMaterial;
+
+		class DATASMITHINTERCHANGE_API FDatasmithReferenceMaterialSelector
+		{
+		public:
+			virtual ~FDatasmithReferenceMaterialSelector() = default;
+
+			virtual bool IsValid() const { return bIsValid; }
+			virtual const TCHAR* GetMaterialPath(EDatasmithReferenceMaterialType MaterialType) const { return nullptr; }
+#if WITH_EDITOR
+			virtual void PostImportProcess(EDatasmithReferenceMaterialType MaterialType, EDatasmithReferenceMaterialQuality MaterialQuality, UMaterialInstanceConstant* MaterialInstance) const {}
+#endif
+
+		protected:
+			bool bIsValid = false;
+		};
+	}
+}
+
