Enabled Datasmith import at runtime:
  - Added basic BP function library to load file thru Datasmith Interchange asynchronously or synchronously
Adjustments to Interchange workflow in PIE and standalone modes and runtime
  - Skipped import of variants at runtime, in PIE and standalone game

diff --git a/Engine/Plugins/Experimental/Enterprise/DatasmithInterchange/DatasmithInterchange.uplugin b/Engine/Plugins/Experimental/Enterprise/DatasmithInterchange/DatasmithInterchange.uplugin
--- a/Engine/Plugins/Experimental/Enterprise/DatasmithInterchange/DatasmithInterchange.uplugin
+++ b/Engine/Plugins/Experimental/Enterprise/DatasmithInterchange/DatasmithInterchange.uplugin
@@ -14,16 +14,21 @@
 	"CanContainContent" : false,
 	"IsBetaVersion" : true,
 	"Installed" : false,
+	"PlatformAllowList": [
+		"Win64",
+		"Linux",
+		"Mac"
+	],
 	"Modules": [
 		{
 			"Name": "DatasmithInterchange",
 			"Type": "Runtime",
-			"LoadingPhase": "Default",
-			"PlatformAllowList": [
-				"Win64",
-				"Linux",
-				"Mac"
-			]
+			"LoadingPhase": "Default"
+		},
+		{
+			"Name": "DatasmithInterchangeScripting",
+			"Type": "Runtime",
+			"LoadingPhase": "Default"
 		}
 	],
 	"Plugins": [
diff --git a/Engine/Plugins/Experimental/Enterprise/DatasmithInterchange/Source/DatasmithInterchange/DatasmithInterchange.Build.cs b/Engine/Plugins/Experimental/Enterprise/DatasmithInterchange/Source/DatasmithInterchange/DatasmithInterchange.Build.cs
--- a/Engine/Plugins/Experimental/Enterprise/DatasmithInterchange/Source/DatasmithInterchange/DatasmithInterchange.Build.cs
+++ b/Engine/Plugins/Experimental/Enterprise/DatasmithInterchange/Source/DatasmithInterchange/DatasmithInterchange.Build.cs
@@ -32,6 +32,7 @@ namespace UnrealBuildTool.Rules
 					{
 						"DatasmithImporter",
 						"DesktopPlatform",
+						"InputCore",
 						"InterchangeEditorPipelines",
 						"Slate",
 						"SlateCore",
diff --git a/Engine/Plugins/Experimental/Enterprise/DatasmithInterchange/Source/DatasmithInterchange/Private/DatasmithInterchangeModule.cpp b/Engine/Plugins/Experimental/Enterprise/DatasmithInterchange/Source/DatasmithInterchange/Private/DatasmithInterchangeModule.cpp
--- a/Engine/Plugins/Experimental/Enterprise/DatasmithInterchange/Source/DatasmithInterchange/Private/DatasmithInterchangeModule.cpp
+++ b/Engine/Plugins/Experimental/Enterprise/DatasmithInterchange/Source/DatasmithInterchange/Private/DatasmithInterchangeModule.cpp
@@ -24,11 +24,15 @@
 
 #if WITH_EDITOR
 #include "DatasmithImporterModule.h"
+#include "DatasmithUtils.h"
 #include "DesktopPlatformModule.h"
 #include "Dialogs/DlgPickPath.h"
 #include "EditorDirectories.h"
 #include "IDesktopPlatform.h"
+#include "Interfaces/IMainFrameModule.h"
 #include "ObjectTools.h"
+#include "Styling/SlateIconFinder.h"
+#include "UI/DatasmithImportOptionsWindow.h"
 #include "ToolMenus.h"
 #include "ToolMenuSection.h"
 #endif //WITH_EDITOR
@@ -88,6 +92,8 @@ private:
 
 	void OnImportInterchange();
 
+	bool DisplayOptionsDialog(const FString& FilePath, const FString& PackagePath, const TArray<TObjectPtr<UDatasmithOptionsBase>>& ImportOptions);
+		
 	bool Import(const FString& FilePath, const FString& ContentPath);
 #endif //WITH_EDITOR
 };
@@ -99,7 +105,8 @@ void FDatasmithInterchangeModule::ExtendDatasmitMenuOptions(FToolMenuSection& Su
 		TEXT("InterchangeImportFile"),
 		LOCTEXT("DatasmithInterchangeFileImport", "Interchange Import..."), // label
 		LOCTEXT("DatasmithInterchangeFileImportTooltip", "Experimental: Import Unreal Datasmith file using Interchange"), // description
-		FSlateIcon(),
+		FSlateIcon(FSlateIconFinder::FindIcon("Datasmith.Import")),
+		//FSlateIcon(FDatasmithStyle::GetStyleSetName(), TEXT("Datasmith.Import")),
 		FUIAction(FExecuteAction::CreateRaw(this, &FDatasmithInterchangeModule::OnImportInterchange))
 	);
 }
@@ -152,6 +159,60 @@ void FDatasmithInterchangeModule::OnImportInterchange()
 	}
 }
 
+bool FDatasmithInterchangeModule::DisplayOptionsDialog(const FString& FilePath, const FString& PackagePath, const TArray<TObjectPtr<UDatasmithOptionsBase>>& ImportOptionsPtr)
+{
+	if (ImportOptionsPtr.Num() == 0)
+	{
+		return true;
+	}
+
+	TSharedPtr<SWindow> ParentWindow;
+
+	if (FModuleManager::Get().IsModuleLoaded("MainFrame"))
+	{
+		IMainFrameModule& MainFrame = FModuleManager::LoadModuleChecked<IMainFrameModule>("MainFrame");
+		ParentWindow = MainFrame.GetParentWindow();
+	}
+
+	TArray<UObject*> ImportOptions;
+	ImportOptions.SetNum(ImportOptionsPtr.Num());
+	for (int32 Index = 0; Index < ImportOptionsPtr.Num(); ++Index)
+	{
+		ImportOptions[Index] = ImportOptionsPtr[Index];
+	}
+
+	TSharedRef<SWindow> Window = SNew(SWindow)
+		.Title(LOCTEXT("DatasmithImportSettingsTitle", "Datasmith Import Options"))
+		.SizingRule(ESizingRule::Autosized);
+
+	float SceneVersion = FDatasmithUtils::GetDatasmithFormatVersionAsFloat();
+	FString FileSDKVersion = FDatasmithUtils::GetEnterpriseVersionAsString();
+
+	TSharedPtr<SDatasmithOptionsWindow> OptionsWindow;
+	Window->SetContent
+	(
+		SAssignNew(OptionsWindow, SDatasmithOptionsWindow)
+		.ImportOptions(ImportOptions)
+		.WidgetWindow(Window)
+		// note: Spacing in text below is intentional for text alignment
+		.FileNameText(FText::Format(LOCTEXT("DatasmithImportSettingsFileName", "  Import File  :    {0}"), FText::FromString(FPaths::GetCleanFilename(FilePath))))
+		.FilePathText(FText::FromString(FilePath))
+		.FileFormatVersion(SceneVersion)
+		.FileSDKVersion(FText::FromString(FileSDKVersion))
+		.PackagePathText(FText::Format(LOCTEXT("DatasmithImportSettingsPackagePath", "  Import To   :    {0}"), FText::FromString(PackagePath)))
+		.ProceedButtonLabel(LOCTEXT("DatasmithOptionWindow_ImportCurLevel", "Import"))
+		.ProceedButtonTooltip(LOCTEXT("DatasmithOptionWindow_ImportCurLevel_ToolTip", "Import the file through Interchange and add to the current Level"))
+		.CancelButtonLabel(LOCTEXT("DatasmithOptionWindow_Cancel", "Cancel"))
+		.CancelButtonTooltip(LOCTEXT("DatasmithOptionWindow_Cancel_ToolTip", "Cancel importing this file"))
+		.MinDetailHeight(320.f)
+		.MinDetailWidth(450.f)
+	);
+
+	FSlateApplication::Get().AddModalWindow(Window, ParentWindow, false);
+
+	return OptionsWindow->ShouldImport();
+}
+
 bool FDatasmithInterchangeModule::Import(const FString& FilePath, const FString& ContentPath)
 {
 	UE::Interchange::FScopedSourceData ScopedSourceData(FilePath);
@@ -161,6 +222,31 @@ bool FDatasmithInterchangeModule::Import(const FString& FilePath, const FString&
 		return false;
 	}
 
+	{
+		FDatasmithSceneSource DatasmithSource;
+		DatasmithSource.SetSourceFile(FilePath);
+
+		TSharedPtr<IDatasmithTranslator> Translator = FDatasmithTranslatorManager::Get().SelectFirstCompatible(DatasmithSource);
+		if (Translator)
+		{
+			TArray<TObjectPtr<UDatasmithOptionsBase>> ImportOptions;
+			Translator->GetSceneImportOptions(ImportOptions);
+
+			bool bShouldImport = DisplayOptionsDialog(FilePath, ContentPath, ImportOptions);
+			if (!bShouldImport)
+			{
+				return false;
+			}
+
+			const FString OptionFilePath = UInterchangeDatasmithTranslator::BuildConfigFilePath(FilePath);
+			for (TObjectPtr<UDatasmithOptionsBase>& Option : ImportOptions)
+			{
+				Option->SaveConfig(CPF_Config, *OptionFilePath);
+			}
+		}
+	}
+
+
 	FImportAssetParameters ImportAssetParameters;
 	ImportAssetParameters.bIsAutomated = false;
 
diff --git a/Engine/Plugins/Experimental/Enterprise/DatasmithInterchange/Source/DatasmithInterchange/Private/InterchangeDatasmithMaterialPipeline.cpp b/Engine/Plugins/Experimental/Enterprise/DatasmithInterchange/Source/DatasmithInterchange/Private/InterchangeDatasmithMaterialPipeline.cpp
--- a/Engine/Plugins/Experimental/Enterprise/DatasmithInterchange/Source/DatasmithInterchange/Private/InterchangeDatasmithMaterialPipeline.cpp
+++ b/Engine/Plugins/Experimental/Enterprise/DatasmithInterchange/Source/DatasmithInterchange/Private/InterchangeDatasmithMaterialPipeline.cpp
@@ -33,7 +33,7 @@ void UInterchangeDatasmithMaterialPipeline::ExecutePreImportPipeline(UInterchang
 {
 	using namespace UE::DatasmithInterchange;
 
-	TArray< UInterchangeDatasmithMaterialNode*> ReferenceMaterials;
+	TArray< UInterchangeDatasmithMaterialNode*> InstancedMaterials;
 	TArray<UInterchangeShaderNode*> ShaderNodes;
 
 	//Find all translated node we need for this pipeline
@@ -41,7 +41,7 @@ void UInterchangeDatasmithMaterialPipeline::ExecutePreImportPipeline(UInterchang
 		{
 			if (UInterchangeDatasmithMaterialNode* MaterialNode = Cast<UInterchangeDatasmithMaterialNode>(Node))
 			{
-				ReferenceMaterials.Add(MaterialNode);
+				InstancedMaterials.Add(MaterialNode);
 			}
 			else if (UInterchangeShaderNode* ShaderNode = Cast<UInterchangeShaderNode>(Node))
 			{
@@ -53,7 +53,7 @@ void UInterchangeDatasmithMaterialPipeline::ExecutePreImportPipeline(UInterchang
 
 	UpdateMaterialFactoryNodes(ShaderNodes);
 
-	for (UInterchangeDatasmithMaterialNode* MaterialNode : ReferenceMaterials)
+	for (UInterchangeDatasmithMaterialNode* MaterialNode : InstancedMaterials)
 	{
 		PreImportMaterialNode(NodeContainer, MaterialNode);
 	}
@@ -133,7 +133,8 @@ void UInterchangeDatasmithMaterialPipeline::PreImportMaterialNode(UInterchangeBa
 	MaterialFactoryNode->SetCustomParent(PackagePath);
 
 #if WITH_EDITOR
-	MaterialFactoryNode->SetCustomInstanceClassName(UMaterialInstanceConstant::StaticClass()->GetPathName());
+	const UClass* MaterialClass = FApp::IsGame() ? UMaterialInstanceDynamic::StaticClass() : UMaterialInstanceConstant::StaticClass();
+	MaterialFactoryNode->SetCustomInstanceClassName(MaterialClass->GetPathName());
 #else
 	MaterialFactoryNode->SetCustomInstanceClassName(UMaterialInstanceDynamic::StaticClass()->GetPathName());
 #endif
diff --git a/Engine/Plugins/Experimental/Enterprise/DatasmithInterchange/Source/DatasmithInterchange/Private/InterchangeDatasmithPipeline.cpp b/Engine/Plugins/Experimental/Enterprise/DatasmithInterchange/Source/DatasmithInterchange/Private/InterchangeDatasmithPipeline.cpp
--- a/Engine/Plugins/Experimental/Enterprise/DatasmithInterchange/Source/DatasmithInterchange/Private/InterchangeDatasmithPipeline.cpp
+++ b/Engine/Plugins/Experimental/Enterprise/DatasmithInterchange/Source/DatasmithInterchange/Private/InterchangeDatasmithPipeline.cpp
@@ -67,30 +67,21 @@ void UInterchangeDatasmithPipeline::ExecutePreImportPipeline(UInterchangeBaseNod
 		return;
 	}
 
-	if (TexturePipeline)
-	{
-		TexturePipeline->ScriptedExecutePreImportPipeline(BaseNodeContainer, SourceDatas);
-	}
-
-	if (MaterialPipeline)
-	{
-		MaterialPipeline->ScriptedExecutePreImportPipeline(BaseNodeContainer, SourceDatas);
-	}
-
-	if (MeshPipeline)
-	{
-		MeshPipeline->ScriptedExecutePreImportPipeline(BaseNodeContainer, SourceDatas);
-	}
-
-	if (LevelPipeline)
-	{
-		LevelPipeline->ScriptedExecutePreImportPipeline(BaseNodeContainer, SourceDatas);
-	}
-
-	if (AnimationPipeline)
+	auto ExecutePreImportPipelineFunc = [this,SourceDatas](UInterchangePipelineBase* Pipeline)
 	{
-		AnimationPipeline->ScriptedExecutePreImportPipeline(BaseNodeContainer, SourceDatas);
-	}
+		if (Pipeline)
+		{
+			Pipeline->SetResultsContainer(this->Results);
+			Pipeline->ScriptedExecutePreImportPipeline(this->BaseNodeContainer, SourceDatas);
+		}
+	};
+
+	ensure(!Results.IsNull());
+	ExecutePreImportPipelineFunc(TexturePipeline);
+	ExecutePreImportPipelineFunc(MaterialPipeline);
+	ExecutePreImportPipelineFunc(MeshPipeline);
+	ExecutePreImportPipelineFunc(LevelPipeline);
+	ExecutePreImportPipelineFunc(AnimationPipeline);
 
 	const UInterchangeDatasmithSceneNode* DatasmithSceneNode = DatasmithSceneNodes[0];
 	const FString PackageSubPath = DatasmithSceneNode->GetDisplayLabel();
diff --git a/Engine/Plugins/Experimental/Enterprise/DatasmithInterchange/Source/DatasmithInterchange/Private/InterchangeDatasmithTranslator.cpp b/Engine/Plugins/Experimental/Enterprise/DatasmithInterchange/Source/DatasmithInterchange/Private/InterchangeDatasmithTranslator.cpp
--- a/Engine/Plugins/Experimental/Enterprise/DatasmithInterchange/Source/DatasmithInterchange/Private/InterchangeDatasmithTranslator.cpp
+++ b/Engine/Plugins/Experimental/Enterprise/DatasmithInterchange/Source/DatasmithInterchange/Private/InterchangeDatasmithTranslator.cpp
@@ -32,6 +32,8 @@
 #include "InterchangeTexture2DNode.h"
 #include "InterchangeVariantSetNode.h"
 
+#include "Misc/App.h"
+
 #define LOCTEXT_NAMESPACE "DatasmithInterchange"
 
 bool UInterchangeDatasmithTranslator::CanImportSourceData(const UInterchangeSourceData* InSourceData) const
@@ -52,13 +54,14 @@ bool UInterchangeDatasmithTranslator::CanImportSourceData(const UInterchangeSour
 	return ExternalSource.IsValid() && ExternalSource->IsAvailable();
 }
 
+#pragma optimize("",off)
 bool UInterchangeDatasmithTranslator::Translate(UInterchangeBaseNodeContainer& BaseNodeContainer) const
 {
 	using namespace UE::DatasmithImporter;
 	using namespace UE::DatasmithInterchange;
 
 	// TODO: This code should eventually go into UInterchangeTranslatorBase once the FExternalSource module gets integrated into Interchange
-	FString FilePath = SourceData->GetFilename();
+	FString FilePath = FPaths::ConvertRelativePathToFull(SourceData->GetFilename());
 	FileName = FPaths::GetCleanFilename(FilePath);
 	const FSourceUri FileNameUri = FSourceUri::FromFilePath(FilePath);
 	const_cast<UInterchangeDatasmithTranslator*>(this)->LoadedExternalSource = IExternalSourceModule::GetOrCreateExternalSource(FileNameUri);
@@ -74,6 +77,26 @@ bool UInterchangeDatasmithTranslator::Translate(UInterchangeBaseNodeContainer& B
 		return false;
 	}
 
+	// Temporary: Update the tessellation options of the associated translator
+	{
+		const TSharedPtr<IDatasmithTranslator>& DatasmithTranslator = LoadedExternalSource->GetAssetTranslator();
+
+		const FString OptionFilePath = BuildConfigFilePath(FilePath);
+
+		if (DatasmithTranslator.IsValid() && FPaths::FileExists(OptionFilePath))
+		{
+			TArray<TObjectPtr<UDatasmithOptionsBase>> ImportOptions;
+			DatasmithTranslator->GetSceneImportOptions(ImportOptions);
+
+			for (TObjectPtr<UDatasmithOptionsBase>& Option : ImportOptions)
+			{
+				Option->LoadConfig(nullptr, *OptionFilePath);
+			}
+
+			DatasmithTranslator->SetSceneImportOptions(ImportOptions);
+		}
+	}
+
 	StartTime = FPlatformTime::Cycles64();
 	FPaths::NormalizeFilename(FilePath);
 
@@ -271,6 +294,8 @@ bool UInterchangeDatasmithTranslator::Translate(UInterchangeBaseNodeContainer& B
 	}
 
 	// Level variant sets
+	// Note: Variants are not supported yet in game play mode
+	if(!FApp::IsGame()) 
 	{
 		const int32 LevelVariantSetCount = DatasmithScene->GetLevelVariantSetsCount();
 
@@ -292,6 +317,7 @@ bool UInterchangeDatasmithTranslator::Translate(UInterchangeBaseNodeContainer& B
 
 	return true;
 }
+#pragma optimize("",on)
 
 void UInterchangeDatasmithTranslator::HandleDatasmithActor(UInterchangeBaseNodeContainer& BaseNodeContainer, const TSharedRef<IDatasmithActorElement>& ActorElement, const UInterchangeSceneNode* ParentNode) const
 {
@@ -569,12 +595,12 @@ TFuture<TOptional<UE::Interchange::FStaticMeshPayloadData>> UInterchangeDatasmit
 
 TFuture<TOptional<UE::Interchange::FAnimationCurvePayloadData>> UInterchangeDatasmithTranslator::GetAnimationCurvePayloadData(const FString& PayLoadKey) const
 {
-	return GetAnimationCurvePayloadData<UE::Interchange::FAnimationCurvePayloadData>(PayLoadKey);
+	return GetAnimationPayloadDataAsCurve<UE::Interchange::FAnimationCurvePayloadData>(PayLoadKey);
 }
 
 TFuture<TOptional<UE::Interchange::FAnimationStepCurvePayloadData>> UInterchangeDatasmithTranslator::GetAnimationStepCurvePayloadData(const FString& PayLoadKey) const
 {
-	return GetAnimationCurvePayloadData<UE::Interchange::FAnimationStepCurvePayloadData>(PayLoadKey);
+	return GetAnimationPayloadDataAsCurve<UE::Interchange::FAnimationStepCurvePayloadData>(PayLoadKey);
 }
 
 TFuture<TOptional<UE::Interchange::FAnimationBakeTransformPayloadData>> UInterchangeDatasmithTranslator::GetAnimationBakeTransformPayloadData(const FString& PayLoadKey, const double BakeFrequency, const double RangeStartSecond, const double RangeStopSecond) const
@@ -646,4 +672,11 @@ void UInterchangeDatasmithTranslator::ImportFinish()
 	UE_LOG(LogInterchangeDatasmith, Log, TEXT("Imported %s in [%d min %.3f s]"), *FileName, ElapsedMin, ElapsedSeconds);
 }
 
+FString UInterchangeDatasmithTranslator::BuildConfigFilePath(const FString& FilePath)
+{
+	const FString OptionFileName = FMD5::HashAnsiString(*(FPaths::ConvertRelativePathToFull(FilePath) + TEXT("_Config"))) + TEXT(".ini");
+	return  FPaths::Combine(FPlatformProcess::UserTempDir(), OptionFileName);
+}
+
+
 #undef LOCTEXT_NAMESPACE
\ No newline at end of file
diff --git a/Engine/Plugins/Experimental/Enterprise/DatasmithInterchange/Source/DatasmithInterchange/Public/InterchangeDatasmithTranslator.h b/Engine/Plugins/Experimental/Enterprise/DatasmithInterchange/Source/DatasmithInterchange/Public/InterchangeDatasmithTranslator.h
--- a/Engine/Plugins/Experimental/Enterprise/DatasmithInterchange/Source/DatasmithInterchange/Public/InterchangeDatasmithTranslator.h
+++ b/Engine/Plugins/Experimental/Enterprise/DatasmithInterchange/Source/DatasmithInterchange/Public/InterchangeDatasmithTranslator.h
@@ -4,8 +4,6 @@
 
 #include "CoreMinimal.h"
 
-#include "InterchangeDatasmithUtils.h"
-
 #include "Animation/InterchangeAnimationPayloadInterface.h"
 #include "InterchangeTranslatorBase.h"
 #include "Mesh/InterchangeStaticMeshPayloadInterface.h"
@@ -45,8 +43,15 @@ namespace UE::DatasmithImporter
 	class FExternalSource;
 }
 
+namespace UE::DatasmithInterchange::AnimUtils
+{
+	typedef TPair<float, TSharedPtr<IDatasmithBaseAnimationElement>> FAnimationPayloadDesc;
+	extern bool GetAnimationPayloadData(const IDatasmithBaseAnimationElement& AnimationElement, float FrameRate, UE::Interchange::FAnimationCurvePayloadData& PayLoadData);
+	extern bool GetAnimationPayloadData(const IDatasmithBaseAnimationElement& AnimationElement, float FrameRate, UE::Interchange::FAnimationStepCurvePayloadData& PayLoadData);
+}
+
 UCLASS(BlueprintType, Experimental)
-class UInterchangeDatasmithTranslator : public UInterchangeTranslatorBase
+class DATASMITHINTERCHANGE_API UInterchangeDatasmithTranslator : public UInterchangeTranslatorBase
 	, public IInterchangeTexturePayloadInterface
 	, public IInterchangeStaticMeshPayloadInterface
 	, public IInterchangeAnimationPayloadInterface
@@ -88,6 +93,9 @@ public:
 	virtual TFuture<TOptional<UE::Interchange::FVariantSetPayloadData>> GetVariantSetPayloadData(const FString& PayloadKey) const override;
 	/* IInterchangeVariantSetPayloadInterface End */
 
+	/** Returns a unique file path to */ 
+	static FString BuildConfigFilePath(const FString& FilePath);
+
 private:
 
 	void HandleDatasmithActor(UInterchangeBaseNodeContainer& BaseNodeContainer, const TSharedRef<IDatasmithActorElement>& ActorElement, const UInterchangeSceneNode* ParentNode) const;
@@ -99,7 +107,7 @@ private:
 	TSharedPtr<UE::DatasmithImporter::FExternalSource> LoadedExternalSource;
 
 	template<typename T>
-	TFuture<TOptional<T>> GetAnimationCurvePayloadData(const FString& PayLoadKey) const
+	TFuture<TOptional<T>> GetAnimationPayloadDataAsCurve(const FString& PayLoadKey) const
 	{
 		using namespace UE::DatasmithInterchange::AnimUtils;
 
diff --git a/Engine/Plugins/Experimental/Enterprise/DatasmithInterchange/Source/DatasmithInterchangeScripting/DatasmithInterchangeScripting.Build.cs b/Engine/Plugins/Experimental/Enterprise/DatasmithInterchange/Source/DatasmithInterchangeScripting/DatasmithInterchangeScripting.Build.cs
new file mode 100644
--- /dev/null
+++ b/Engine/Plugins/Experimental/Enterprise/DatasmithInterchange/Source/DatasmithInterchangeScripting/DatasmithInterchangeScripting.Build.cs
@@ -0,0 +1,34 @@
+// Copyright Epic Games, Inc. All Rights Reserved.
+
+namespace UnrealBuildTool.Rules
+{	public class DatasmithInterchangeScripting : ModuleRules
+	{
+		public DatasmithInterchangeScripting(ReadOnlyTargetRules Target)
+			: base(Target)
+		{
+
+			PrivateDependencyModuleNames.AddRange(
+				new string[]
+				{
+					"DatasmithInterchange",
+					"DatasmithContent",
+					"DatasmithCore",
+					"DatasmithTranslator",
+					"ExternalSource",
+					"InterchangeEngine",
+					"InterchangeCommonParser",
+				}
+			);
+
+			PublicDependencyModuleNames.AddRange(
+				new string[]
+				{
+					"Core",
+					"CoreUObject",
+					"Engine",
+					"InterchangeCore",
+				}
+			);
+		}
+	}
+}
\ No newline at end of file
diff --git a/Engine/Plugins/Experimental/Enterprise/DatasmithInterchange/Source/DatasmithInterchangeScripting/Private/DatasmithInterchangeScripting.cpp b/Engine/Plugins/Experimental/Enterprise/DatasmithInterchange/Source/DatasmithInterchangeScripting/Private/DatasmithInterchangeScripting.cpp
new file mode 100644
--- /dev/null
+++ b/Engine/Plugins/Experimental/Enterprise/DatasmithInterchange/Source/DatasmithInterchangeScripting/Private/DatasmithInterchangeScripting.cpp
@@ -0,0 +1,270 @@
+// Copyright Epic Games, Inc. All Rights Reserved.
+
+#include "DatasmithInterchangeScripting.h"
+
+#include "InterchangeDatasmithPipeline.h"
+#include "InterchangeDatasmithTranslator.h"
+
+#include "DatasmithImportOptions.h"
+#include "DatasmithTranslatorManager.h"
+
+#include "Async/Async.h"
+#include "Engine/World.h"
+#include "GameFramework/Actor.h"
+#include "HAL/FileManager.h"
+#include "Logging/LogMacros.h"
+#include "Misc/App.h"
+#include "Misc/Paths.h"
+
+DECLARE_LOG_CATEGORY_EXTERN(LogDatasmithInterchangeImportLibrary, Log, All)
+DEFINE_LOG_CATEGORY(LogDatasmithInterchangeImportLibrary)
+
+#define LOCTEXT_NAMESPACE "DatasmithInterchangeImportLibrary"
+
+FDatasmithInterchangeImportContext::FDatasmithInterchangeImportContext()
+	: AssetPath(TEXT("/Game/DatasmithInterchangeImport"))
+{
+}
+
+UDatasmithInterchangeImportResult::~UDatasmithInterchangeImportResult()
+{
+	const FString ConfigFilePath = UInterchangeDatasmithTranslator::BuildConfigFilePath(FilePath);
+	if (FPaths::FileExists(ConfigFilePath))
+	{
+		IFileManager::Get().Delete(*ConfigFilePath);
+	}
+}
+
+bool UDatasmithInterchangeImportResult::Build(const FString& InFilePath, FImportResultPtr InAssetResult, FImportResultPtr InSceneResult, bool bAsync, USceneComponent* Anchor)
+{
+	using namespace UE::Interchange;
+
+	FilePath = InFilePath;
+	AssetResult = InAssetResult;
+	SceneResult = InSceneResult;
+
+	bIsValid = AssetResult.IsValid() && AssetResult->IsValid() && SceneResult.IsValid(); // Status of SceneResult is not set yet.
+	
+	if (bIsValid)
+	{
+		// Note: All the logic below assume that this is valid at the end of the import
+		if (bAsync)
+		{
+			// TODO: The two results might be done by the time this code is executed
+			SceneResult->OnDone(
+				[this, Anchor](UE::Interchange::FImportResult& Result)
+				{
+					ImportedObjects.Append(Result.GetImportedObjects());
+					if (AssetResult->GetStatus() == FImportResult::EStatus::Done)
+					{
+						UE_LOG(LogDatasmithInterchangeImportLibrary, Warning, TEXT("Scene import completed"));
+						OnImportComplete(Anchor, true);
+					}
+				}
+			);
+
+			AssetResult->OnDone(
+				[this, Anchor](UE::Interchange::FImportResult& Result)
+				{
+					ImportedObjects.Append(Result.GetImportedObjects());
+					if (SceneResult->GetStatus() == FImportResult::EStatus::Done)
+					{
+						UE_LOG(LogDatasmithInterchangeImportLibrary, Warning, TEXT("Asset import completed"));
+						OnImportComplete(Anchor, true);
+					}
+				}
+			);
+		}
+		else
+		{
+			// Wait for the import to complete
+			AssetResult->WaitUntilDone();
+			SceneResult->WaitUntilDone();
+
+			ImportedObjects.Append(AssetResult->GetImportedObjects());
+			ImportedObjects.Append(SceneResult->GetImportedObjects());
+
+			OnImportComplete(Anchor, true);
+		}
+	}
+
+	return bIsValid;
+}
+
+void UDatasmithInterchangeImportResult::OnImportComplete(USceneComponent* Anchor, bool bAsync)
+{
+	auto ApplyAnchor = [Anchor](const TArray< UObject* >& ResultObjects)
+	{
+		if (Anchor)
+		{
+#if WITH_EDITOR
+			UE_LOG(LogDatasmithInterchangeImportLibrary, Warning, TEXT("Applying anchor..."));
+			int32 RootLessCount = 0;
+#endif
+			for (UObject* Object : ResultObjects)
+			{
+				if (AActor* Actor = Cast<AActor>(Object))
+				{
+					if (!Actor->GetAttachParentActor() && Actor->GetRootComponent())
+					{
+#if WITH_EDITOR
+						if (RootLessCount > 0)
+						{
+							UE_LOG(LogDatasmithInterchangeImportLibrary, Warning, TEXT("Actor %s is rootless. It should not..."), *Actor->GetActorLabel());
+						}
+						RootLessCount++;
+#endif
+						Anchor->GetOwner()->UnregisterAllComponents(true);
+						Actor->UnregisterAllComponents(true);
+
+						USceneComponent* RootComponent = Actor->GetRootComponent();
+						if (RootComponent->Mobility == EComponentMobility::Static && Anchor->Mobility != EComponentMobility::Static)
+						{
+							RootComponent->SetMobility(EComponentMobility::Movable);
+						}
+
+						Actor->GetRootComponent()->AttachToComponent(Anchor, FAttachmentTransformRules::KeepRelativeTransform);
+
+						Actor->ReregisterAllComponents();
+						Anchor->GetOwner()->ReregisterAllComponents();
+					}
+					//break;
+				}
+			}
+		}
+	};
+
+	// Trigger notification on next tick on game thread
+	AsyncTask(ENamedThreads::GameThread, [this, ApplyAnchor] {
+			ApplyAnchor(this->ImportedObjects);
+			this->OnImportEnded.Broadcast();
+		});
+}
+
+UDatasmithInterchangeImportResult* UDatasmithInterchangeScripting::LoadFile(const FString& FilePath, const FDatasmithInterchangeImportContext& Context)
+{
+	using FInterchangeResults = TTuple<UE::Interchange::FAssetImportResultRef, UE::Interchange::FSceneImportResultRef>;
+
+	UDatasmithInterchangeImportResult* Result = nullptr;
+
+	UE::Interchange::FScopedSourceData ScopedSourceData(FilePath);
+
+	UInterchangeManager& InterchangeManager = UInterchangeManager::GetInterchangeManager();
+
+	if (!InterchangeManager.CanTranslateSourceData(ScopedSourceData.GetSourceData()))
+	{
+		return Result;
+	}
+
+	// Save tessellation options' settings, so they can be used by the Datasmith translator
+	{
+		const FString ConfigFilePath = UInterchangeDatasmithTranslator::BuildConfigFilePath(FilePath);
+
+		for (UObject* Option : Context.ImportOptions)
+		{
+			if (Option)
+			{
+				Option->SaveConfig(CPF_Config, *ConfigFilePath);
+			}
+		}
+	}
+
+	UInterchangeDatasmithPipeline* DatasmithPipeline = NewObject< UInterchangeDatasmithPipeline >();
+
+	FImportAssetParameters ImportAssetParameters;
+	ImportAssetParameters.bIsAutomated = false;
+
+	ImportAssetParameters.OverridePipelines.Add(DatasmithPipeline);
+
+	FString AssetPath = FApp::IsGame() ? FPaths::Combine(GetTransientPackage()->GetPathName(), FGuid::NewGuid().ToString()) : Context.AssetPath;
+
+	FInterchangeResults ImportResults = InterchangeManager.ImportSceneAsync(AssetPath, ScopedSourceData.GetSourceData(), ImportAssetParameters);
+
+	Result = NewObject<UDatasmithInterchangeImportResult>();
+
+	if (!Result->Build(FilePath, ImportResults.Get<0>(), ImportResults.Get<1>(), Context.bAsync, Context.Anchor))
+	{
+		return nullptr;
+	}
+
+	return Result;
+}
+
+void UDatasmithInterchangeScripting::GetDatasmithFormats(FString& Extensions, FString& FileTypes)
+{
+	const TArray<FString> Formats = FDatasmithTranslatorManager::Get().GetSupportedFormats();
+
+	Extensions = TEXT("");
+
+	TMap<FString, FString> FileTypeMap;
+
+	for (TArray<FString>::TConstIterator FormatIter(Formats); FormatIter; ++FormatIter)
+	{
+		const FString& CurFormat = *FormatIter;
+
+		// Parse the format into its extension and description parts
+		TArray<FString> FormatComponents;
+		CurFormat.ParseIntoArray(FormatComponents, TEXT(";"), false);
+
+		for (int32 ComponentIndex = 0; ComponentIndex < FormatComponents.Num(); ComponentIndex += 2)
+		{
+			const FString CurExtension = FString::Printf(TEXT("*.%s"), *FormatComponents[ComponentIndex]);
+			if (CurExtension.Equals(TEXT(".gltf"), ESearchCase::IgnoreCase) || CurExtension.Equals(TEXT(".glb"), ESearchCase::IgnoreCase))
+			{
+				continue;
+			}
+
+			const FString& CurDescription = FormatComponents[ComponentIndex + 1];
+
+			if (Extensions.Len() > 0)
+			{
+				Extensions += TEXT(";");
+			}
+
+			Extensions += CurExtension;
+
+			if (CurDescription.Len() > 0)
+			{
+				FString& FileTypeExtensions = FileTypeMap.FindOrAdd(CurDescription);
+
+				if (FileTypeExtensions.Len() > 0)
+				{
+					FileTypeExtensions += TEXT("; ");
+				}
+
+				FileTypeExtensions += CurExtension;
+
+			}
+		}
+	}
+
+	FileTypes = FString::Printf(TEXT("All Files (%s)|%s"), *Extensions, *Extensions);
+
+	for (const TPair<FString, FString>& Entry : FileTypeMap)
+	{
+		FileTypes += FString::Printf(TEXT("|%s (%s)|%s"), *Entry.Key, *Entry.Value, *Entry.Value);
+	}
+
+}
+
+void UDatasmithInterchangeScripting::GetDatasmithOptionsForFile(const FString& FilePath, TArray<UObject*>& Options)
+{
+	FDatasmithSceneSource Source;
+	Source.SetSourceFile(FilePath);
+
+	if (TSharedPtr<IDatasmithTranslator> Translator = FDatasmithTranslatorManager::Get().SelectFirstCompatible(Source))
+	{
+		TArray<TObjectPtr<UDatasmithOptionsBase>> OptionPtrs;
+		Translator->GetSceneImportOptions(OptionPtrs);
+
+		Options.Reserve(Options.Num() + OptionPtrs.Num());
+
+		for (const TObjectPtr<UDatasmithOptionsBase>& Option : OptionPtrs)
+		{
+			if (Option)
+			{
+				Options.Add(Option);
+			}
+		}
+	}
+}
diff --git a/Engine/Plugins/Experimental/Enterprise/DatasmithInterchange/Source/DatasmithInterchangeScripting/Private/DatasmithInterchangeScriptingModule.cpp b/Engine/Plugins/Experimental/Enterprise/DatasmithInterchange/Source/DatasmithInterchangeScripting/Private/DatasmithInterchangeScriptingModule.cpp
new file mode 100644
--- /dev/null
+++ b/Engine/Plugins/Experimental/Enterprise/DatasmithInterchange/Source/DatasmithInterchangeScripting/Private/DatasmithInterchangeScriptingModule.cpp
@@ -0,0 +1,5 @@
+// Copyright Epic Games, Inc. All Rights Reserved.
+
+#include "Modules/ModuleManager.h"
+
+IMPLEMENT_MODULE(FDefaultModuleImpl, DatasmithInterchangeScripting);
diff --git a/Engine/Plugins/Experimental/Enterprise/DatasmithInterchange/Source/DatasmithInterchangeScripting/Public/DatasmithInterchangeScripting.h b/Engine/Plugins/Experimental/Enterprise/DatasmithInterchange/Source/DatasmithInterchangeScripting/Public/DatasmithInterchangeScripting.h
new file mode 100644
--- /dev/null
+++ b/Engine/Plugins/Experimental/Enterprise/DatasmithInterchange/Source/DatasmithInterchangeScripting/Public/DatasmithInterchangeScripting.h
@@ -0,0 +1,111 @@
+// Copyright Epic Games, Inc. All Rights Reserved.
+
+#pragma once
+
+#include "DatasmithImportOptions.h"
+
+#include "Delegates/Delegate.h"
+#include "Delegates/DelegateCombinations.h"
+#include "InterchangeManager.h"
+
+#include "Kismet/BlueprintFunctionLibrary.h"
+
+#include "DatasmithInterchangeScripting.generated.h"
+
+class AActor;
+
+DECLARE_DYNAMIC_MULTICAST_DELEGATE(FDatasmithInterchangeImportEnded);
+
+USTRUCT(BlueprintType, Experimental)
+struct DATASMITHINTERCHANGESCRIPTING_API FDatasmithInterchangeImportContext
+{
+	GENERATED_BODY()
+	
+	FDatasmithInterchangeImportContext();
+
+	/** Path of content folder where to store created assets */
+	UPROPERTY(BlueprintReadWrite, Category = "Datasmith | Interchange")
+	FString AssetPath;
+
+	/** If not null, the resulting ADatasmithSceneActor will be parented to the Anchor */
+	UPROPERTY(BlueprintReadWrite, Category = "Datasmith | Interchange")
+	TObjectPtr<USceneComponent> Anchor;
+
+	/** If true, the import will be asynchronous */
+	UPROPERTY(BlueprintReadWrite, Category = "Datasmith | Interchange")
+	bool bAsync = true;
+
+	/**  */
+	UPROPERTY(EditAnywhere, BlueprintReadWrite, Category = "StaticMesh Options", meta = (ShowOnlyInnerProperties))
+	FDatasmithStaticMeshImportOptions StaticMeshOptions;
+	
+	UPROPERTY(EditAnywhere, BlueprintReadWrite, Category = "Import Options", meta = (ShowOnlyInnerProperties))
+	TArray<TObjectPtr<UObject>> ImportOptions;
+};
+
+UCLASS(BlueprintType, Experimental)
+class DATASMITHINTERCHANGESCRIPTING_API UDatasmithInterchangeImportResult : public UObject
+{
+	GENERATED_BODY()
+
+	using FImportResultPtr = TSharedPtr< UE::Interchange::FImportResult, ESPMode::ThreadSafe >;
+
+public:
+	UDatasmithInterchangeImportResult() {}
+	~UDatasmithInterchangeImportResult();
+
+	bool Build(const FString& InFilePath, FImportResultPtr InAssetResult, FImportResultPtr InSceneResult, bool bAsync, USceneComponent* Anchor);
+
+	bool IsValid() { return bIsValid; }
+
+	/** Dynamic delegate used to trigger an event in the Game when the import is completed */
+	UPROPERTY(EditAnywhere, BlueprintReadOnly, Category = "Datasmith | Interchange")
+	TArray<TObjectPtr<UObject>> ImportedObjects;
+
+	/** Dynamic delegate used to trigger an event in the Game when the import is completed */
+	UPROPERTY(VisibleAnywhere, BlueprintAssignable, Category = "Datasmith | Interchange")
+	FDatasmithInterchangeImportEnded OnImportEnded;
+
+private:
+	void OnImportComplete(USceneComponent* Anchor, bool bAsync);
+
+private:
+	FString FilePath;
+	FImportResultPtr AssetResult;
+	FImportResultPtr SceneResult;
+	bool bIsValid = false;
+};
+
+UCLASS(MinimalAPI, Experimental, meta = (ScriptName = "DatasmithInterchangeScripting"))
+class UDatasmithInterchangeScripting : public UBlueprintFunctionLibrary
+{
+	GENERATED_BODY()
+
+public:
+	/**
+	 * Utility method to load a specific file into the current level
+	 * @param FilePath  : Path to the file to load.
+	 * @param AssetPath	: .
+	 * @param bAsync    : 
+	 * @param Anchor    : 
+	 */
+	UFUNCTION(BlueprintCallable, Category = "Datasmith | Interchange")
+	static UDatasmithInterchangeImportResult* LoadFile(const FString& FilePath, const FDatasmithInterchangeImportContext& Context);
+
+	/**
+	 * Fetch the typed options associated with this file
+	 * @param FilePath	The path to the file to get the option for
+	 * @param OptionType	Type of the option structure to get
+	 * @return The option structure found
+	 */
+	UFUNCTION(BlueprintCallable, Category = "Datasmith | Interchange")
+	void GetDatasmithOptionsForFile(const FString& FilePath, TArray<UObject*>& Options);
+
+	/**
+	 * Utility method to get list of supported extensions and string to pass to IDesktopPlatform::OpenFileDialog
+	 * @param Extensions : ';' separated list of extensions supported by Datasmith thru Interchange
+	 * @param FileTypes	 : The type filters to show in the dialog when using @see IDesktopPlatform::OpenFileDialog.
+	 */
+	UFUNCTION(BlueprintCallable, Category = "Datasmith | Interchange")
+	static void GetDatasmithFormats(FString& Extensions, FString& FileTypes);
+};
\ No newline at end of file
