Fixed assert when trying to load gltf file through Datasmith interchange

diff --git a/Engine/Plugins/Experimental/Enterprise/DatasmithInterchange/Source/DatasmithInterchange/Private/DatasmithInterchangeModule.cpp b/Engine/Plugins/Experimental/Enterprise/DatasmithInterchange/Source/DatasmithInterchange/Private/DatasmithInterchangeModule.cpp
--- a/Engine/Plugins/Experimental/Enterprise/DatasmithInterchange/Source/DatasmithInterchange/Private/DatasmithInterchangeModule.cpp
+++ b/Engine/Plugins/Experimental/Enterprise/DatasmithInterchange/Source/DatasmithInterchange/Private/DatasmithInterchangeModule.cpp
@@ -113,7 +113,20 @@ void FDatasmithInterchangeModule::ExtendDatasmitMenuOptions(FToolMenuSection& Su
 
 void FDatasmithInterchangeModule::OnImportInterchange()
 {
-	const TArray<FString> Formats = FDatasmithTranslatorManager::Get().GetSupportedFormats();
+	const TArray<FString> OldFormats = FDatasmithTranslatorManager::Get().GetSupportedFormats();
+	TArray<FString> Formats;
+	Formats.Reserve(OldFormats.Num());
+
+	for (const FString& Format : OldFormats)
+	{
+		if (Format.Contains(TEXT("gltf")) || Format.Contains(TEXT("glb")))
+		{
+			continue;
+		}
+
+		Formats.Add(Format);
+	}
+
 	FString FileTypes;
 	FString Extensions;
 
@@ -217,8 +230,12 @@ bool FDatasmithInterchangeModule::Import(const FString& FilePath, const FString&
 {
 	UE::Interchange::FScopedSourceData ScopedSourceData(FilePath);
 
-	if (!UInterchangeManager::GetInterchangeManager().CanTranslateSourceData(ScopedSourceData.GetSourceData()))
+	UInterchangeTranslatorBase* InterchangeTranslator = UInterchangeManager::GetInterchangeManager().GetTranslatorForSourceData(ScopedSourceData.GetSourceData());
+	if (!InterchangeTranslator || !InterchangeTranslator->IsA(UInterchangeDatasmithTranslator::StaticClass()))
 	{
+		const FString FileExtension = FPaths::GetExtension(FilePath);
+		UE_LOG(LogInterchangeDatasmith, Warning, TEXT("File extension \"%s\" is not supported by DatasmithInterchange"), *FileExtension);
+
 		return false;
 	}
 
diff --git a/Engine/Plugins/Experimental/Enterprise/DatasmithInterchange/Source/DatasmithInterchange/Private/InterchangeDatasmithTranslator.cpp b/Engine/Plugins/Experimental/Enterprise/DatasmithInterchange/Source/DatasmithInterchange/Private/InterchangeDatasmithTranslator.cpp
--- a/Engine/Plugins/Experimental/Enterprise/DatasmithInterchange/Source/DatasmithInterchange/Private/InterchangeDatasmithTranslator.cpp
+++ b/Engine/Plugins/Experimental/Enterprise/DatasmithInterchange/Source/DatasmithInterchange/Private/InterchangeDatasmithTranslator.cpp
@@ -42,7 +42,7 @@ bool UInterchangeDatasmithTranslator::CanImportSourceData(const UInterchangeSour
 
 	const FString FilePath = InSourceData->GetFilename();
 	const FString FileExtension = FPaths::GetExtension(FilePath);
-	if (FileExtension == TEXT("gltf") || FileExtension == TEXT("glb"))
+	if (FileExtension.Equals(TEXT("gltf"), ESearchCase::IgnoreCase) || FileExtension.Equals(TEXT("glb"), ESearchCase::IgnoreCase))
 	{
 		// Do not translate gltf since there is already a native gltf interchange translator. 
 		return false;
@@ -54,7 +54,6 @@ bool UInterchangeDatasmithTranslator::CanImportSourceData(const UInterchangeSour
 	return ExternalSource.IsValid() && ExternalSource->IsAvailable();
 }
 
-#pragma optimize("",off)
 bool UInterchangeDatasmithTranslator::Translate(UInterchangeBaseNodeContainer& BaseNodeContainer) const
 {
 	using namespace UE::DatasmithImporter;
@@ -317,7 +316,6 @@ bool UInterchangeDatasmithTranslator::Translate(UInterchangeBaseNodeContainer& B
 
 	return true;
 }
-#pragma optimize("",on)
 
 void UInterchangeDatasmithTranslator::HandleDatasmithActor(UInterchangeBaseNodeContainer& BaseNodeContainer, const TSharedRef<IDatasmithActorElement>& ActorElement, const UInterchangeSceneNode* ParentNode) const
 {
