Fixed crash when importing a Datasmith scene using the property UseIESBrightness

diff --git a/Engine/Plugins/Experimental/Enterprise/DatasmithInterchange/Source/DatasmithInterchange/Private/InterchangeDatasmithAreaLightFactoryNode.cpp b/Engine/Plugins/Experimental/Enterprise/DatasmithInterchange/Source/DatasmithInterchange/Private/InterchangeDatasmithAreaLightFactoryNode.cpp
--- a/Engine/Plugins/Experimental/Enterprise/DatasmithInterchange/Source/DatasmithInterchange/Private/InterchangeDatasmithAreaLightFactoryNode.cpp
+++ b/Engine/Plugins/Experimental/Enterprise/DatasmithInterchange/Source/DatasmithInterchange/Private/InterchangeDatasmithAreaLightFactoryNode.cpp
@@ -13,7 +13,7 @@ UE::Interchange::FAttributeKey UInterchangeDatasmithAreaLightFactoryNode::Macro_
 UE::Interchange::FAttributeKey UInterchangeDatasmithAreaLightFactoryNode::Macro_CustomColorKey = UE::Interchange::FAttributeKey(TEXT("Color"));
 UE::Interchange::FAttributeKey UInterchangeDatasmithAreaLightFactoryNode::Macro_CustomTemperatureKey = UE::Interchange::FAttributeKey(TEXT("Temperature"));
 UE::Interchange::FAttributeKey UInterchangeDatasmithAreaLightFactoryNode::Macro_CustomIESTextureKey = UE::Interchange::FAttributeKey(TEXT("IESTexture"));
-UE::Interchange::FAttributeKey UInterchangeDatasmithAreaLightFactoryNode::Macro_CustomUseIESBrightnessKey = UE::Interchange::FAttributeKey(TEXT("UseIESBrightness"));
+UE::Interchange::FAttributeKey UInterchangeDatasmithAreaLightFactoryNode::Macro_CustombUseIESBrightnessKey = UE::Interchange::FAttributeKey(TEXT("UseIESBrightness"));
 UE::Interchange::FAttributeKey UInterchangeDatasmithAreaLightFactoryNode::Macro_CustomIESBrightnessScaleKey = UE::Interchange::FAttributeKey(TEXT("IESBrightnessScale"));
 UE::Interchange::FAttributeKey UInterchangeDatasmithAreaLightFactoryNode::Macro_CustomRotationKey = UE::Interchange::FAttributeKey(TEXT("Rotation"));
 UE::Interchange::FAttributeKey UInterchangeDatasmithAreaLightFactoryNode::Macro_CustomSourceRadiusKey = UE::Interchange::FAttributeKey(TEXT("SourceRadius"));
@@ -136,12 +136,12 @@ bool UInterchangeDatasmithAreaLightFactoryNode::SetCustomIESTexture(const FStrin
 
 bool UInterchangeDatasmithAreaLightFactoryNode::GetCustomUseIESBrightness(bool& AttributeValue) const
 {
-	IMPLEMENT_NODE_ATTRIBUTE_GETTER(UseIESBrightness, bool);
+	IMPLEMENT_NODE_ATTRIBUTE_GETTER(bUseIESBrightness, bool);
 }
 
 bool UInterchangeDatasmithAreaLightFactoryNode::SetCustomUseIESBrightness(const bool& AttributeValue, bool bAddApplyDelegate)
 {
-	IMPLEMENT_NODE_ATTRIBUTE_SETTER(UInterchangeDatasmithAreaLightFactoryNode, UseIESBrightness, bool, ADatasmithAreaLightActor);
+	IMPLEMENT_NODE_ATTRIBUTE_SETTER(UInterchangeDatasmithAreaLightFactoryNode, bUseIESBrightness, bool, ADatasmithAreaLightActor);
 }
 
 bool UInterchangeDatasmithAreaLightFactoryNode::GetCustomIESBrightnessScale(float& AttributeValue) const
diff --git a/Engine/Plugins/Experimental/Enterprise/DatasmithInterchange/Source/DatasmithInterchange/Public/InterchangeDatasmithAreaLightFactoryNode.h b/Engine/Plugins/Experimental/Enterprise/DatasmithInterchange/Source/DatasmithInterchange/Public/InterchangeDatasmithAreaLightFactoryNode.h
--- a/Engine/Plugins/Experimental/Enterprise/DatasmithInterchange/Source/DatasmithInterchange/Public/InterchangeDatasmithAreaLightFactoryNode.h
+++ b/Engine/Plugins/Experimental/Enterprise/DatasmithInterchange/Source/DatasmithInterchange/Public/InterchangeDatasmithAreaLightFactoryNode.h
@@ -136,7 +136,7 @@ private:
 	static  UE::Interchange::FAttributeKey Macro_CustomColorKey;
 	static  UE::Interchange::FAttributeKey Macro_CustomTemperatureKey;
 	static  UE::Interchange::FAttributeKey Macro_CustomIESTextureKey;
-	static  UE::Interchange::FAttributeKey Macro_CustomUseIESBrightnessKey;
+	static  UE::Interchange::FAttributeKey Macro_CustombUseIESBrightnessKey;
 	static  UE::Interchange::FAttributeKey Macro_CustomIESBrightnessScaleKey;
 	static  UE::Interchange::FAttributeKey Macro_CustomRotationKey;
 	static  UE::Interchange::FAttributeKey Macro_CustomSourceRadiusKey;
