Added Datasmith Interchange Runtime load to QAEnterprise (WIP)

Note: Set your GameDefaultMap to /Game/Manual/Interchange/InterchangeTestMap.InterchangeTestMap in Sandbox\Enterprise\QAEnterprise\Config\DefaultEngine.ini

diff --git a/Engine/Plugins/Experimental/Enterprise/DatasmithInterchange/Source/DatasmithInterchange/Private/DatasmithInterchangeModule.cpp b/Engine/Plugins/Experimental/Enterprise/DatasmithInterchange/Source/DatasmithInterchange/Private/DatasmithInterchangeModule.cpp
--- a/Engine/Plugins/Experimental/Enterprise/DatasmithInterchange/Source/DatasmithInterchange/Private/DatasmithInterchangeModule.cpp
+++ b/Engine/Plugins/Experimental/Enterprise/DatasmithInterchange/Source/DatasmithInterchange/Private/DatasmithInterchangeModule.cpp
@@ -51,7 +51,7 @@ public:
 
 		// Load the blueprint asset into memory while wew're on the game thread so that GetAreaLightActorBPClass() can safely be called from other threads.
 		UBlueprint* AreaLightBlueprint = Cast<UBlueprint>(FSoftObjectPath(TEXT("/DatasmithContent/Datasmith/DatasmithArealight.DatasmithArealight")).TryLoad());
-		ensure(AreaLightBlueprint != nullptr);
+		//ensure(AreaLightBlueprint != nullptr);
 
 		UInterchangeManager& InterchangeManager = UInterchangeManager::GetInterchangeManager();
 		InterchangeManager.RegisterTranslator(UInterchangeDatasmithTranslator::StaticClass());
