Fixed warning in "Run QAEnterprise BootTest Win64" step

diff --git a/Engine/Plugins/Experimental/Enterprise/DatasmithInterchange/Source/DatasmithInterchange/Private/InterchangeDatasmithAreaLightFactory.cpp b/Engine/Plugins/Experimental/Enterprise/DatasmithInterchange/Source/DatasmithInterchange/Private/InterchangeDatasmithAreaLightFactory.cpp
--- a/Engine/Plugins/Experimental/Enterprise/DatasmithInterchange/Source/DatasmithInterchange/Private/InterchangeDatasmithAreaLightFactory.cpp
+++ b/Engine/Plugins/Experimental/Enterprise/DatasmithInterchange/Source/DatasmithInterchange/Private/InterchangeDatasmithAreaLightFactory.cpp
@@ -19,7 +19,8 @@ UClass* UInterchangeDatasmithAreaLightFactory::GetFactoryClass() const
 	UBlueprint* LightShapeBlueprint = Cast< UBlueprint >(LightShapeBlueprintRef.ResolveObject());
 
 	// LightShapeBlueprint should already be loaded at this point.
-	if (!ensure(LightShapeBlueprint))
+	//if (!ensure(LightShapeBlueprint))
+	if (!LightShapeBlueprint)
 	{
 		return nullptr;
 	}
