Fixed missing materials with CAD files at runtime: Activated the code procesing Datasmith Pbr materials at runtime.

diff --git a/Engine/Plugins/Experimental/Enterprise/DatasmithInterchange/Source/DatasmithInterchange/Private/InterchangeDatasmithMaterialNode.cpp b/Engine/Plugins/Experimental/Enterprise/DatasmithInterchange/Source/DatasmithInterchange/Private/InterchangeDatasmithMaterialNode.cpp
--- a/Engine/Plugins/Experimental/Enterprise/DatasmithInterchange/Source/DatasmithInterchange/Private/InterchangeDatasmithMaterialNode.cpp
+++ b/Engine/Plugins/Experimental/Enterprise/DatasmithInterchange/Source/DatasmithInterchange/Private/InterchangeDatasmithMaterialNode.cpp
@@ -43,7 +43,6 @@ namespace UE::DatasmithInterchange::MaterialUtils
 	const FName MaterialFunctionPathAttrName(TEXT("Datasmith:Material:FunctionCall:MaterialFunctionPath"));
 	const FName DefaultOutputIndexAttrName(TEXT("Datasmith:MaterialExpression:DefaultOutputIndex"));
 
-#if WITH_EDITOR
 	class FPbrMaterialHelper
 	{
 	public:
@@ -122,7 +121,6 @@ namespace UE::DatasmithInterchange::MaterialUtils
 
 		EBlendMode GetUEPbrImportBlendMode();
 	};
-#endif
 
 	bool BuildMaterialNode(IDatasmithMaterialInstanceElement& MaterialElement, UInterchangeDatasmithMaterialNode& MaterialNode)
 	{
@@ -196,7 +194,6 @@ namespace UE::DatasmithInterchange::MaterialUtils
 	}
 
 
-#if WITH_EDITOR
 	FPbrMaterialHelper::FPbrMaterialHelper(UInterchangeShaderNode& MaterialNode, IDatasmithUEPbrMaterialElement& InMaterialElement)
 		: NodeContainer(*Cast<UInterchangeBaseNodeContainer>(MaterialNode.GetOuter()))
 		, MaterialID(MaterialNode.GetUniqueID())
@@ -372,6 +369,9 @@ namespace UE::DatasmithInterchange::MaterialUtils
 
 	UInterchangeShaderNode* FPbrMaterialHelper::BuildGenericExpression(const IDatasmithMaterialExpressionGeneric& DatasmithExpression, UInterchangeShaderNode& ParentNode, TArray<FName>* OutputNames)
 	{
+		UInterchangeShaderNode* ExpressionNode = nullptr;
+
+#if WITH_EDITOR
 		UMaterialExpression* MaterialExpression = GetDefaultMaterialExpression(DatasmithExpression.GetExpressionName());
 
 		if (!ensure(MaterialExpression))
@@ -380,7 +380,7 @@ namespace UE::DatasmithInterchange::MaterialUtils
 			return nullptr;
 		}
 
-		UInterchangeShaderNode* ExpressionNode = CreateExpressionNode(DatasmithExpression, ParentNode.GetUniqueID());
+		ExpressionNode = CreateExpressionNode(DatasmithExpression, ParentNode.GetUniqueID());
 		if (!ensure(ExpressionNode))
 		{
 				// TODO: Log error
@@ -413,12 +413,14 @@ namespace UE::DatasmithInterchange::MaterialUtils
 				OutputNames->Add(ExpressionOutput.OutputName);
 			}
 		}
+#endif
 
 		return ExpressionNode;
 	}
 
 	void FPbrMaterialHelper::ConnectExpression(const IDatasmithMaterialExpressionGeneric& DatasmithExpression, const FConnectionData& ConnectionData, int32 OutputIndex)
 	{
+#if WITH_EDITOR
 		TArray<FName> OutputNames;
 		UInterchangeShaderNode* ExpressionNode = BuildGenericExpression(DatasmithExpression, *ConnectionData.InputNode, &OutputNames);
 		if (!ensure(ExpressionNode))
@@ -438,10 +440,12 @@ namespace UE::DatasmithInterchange::MaterialUtils
 		// TODO: Handle case of Mask
 		FExpressionOutput& ExpressionOutput = MaterialExpression->GetOutputs()[OutputIndex];
 		ConnectionData.ConnectOuputToInput(ExpressionNode->GetUniqueID(), ExpressionOutput.OutputName);
+#endif
 	}
 
 	void FPbrMaterialHelper::ConnectExpression(const IDatasmithMaterialExpressionFunctionCall& DatasmithExpression, const FConnectionData& ConnectionData, int32 OutputIndex)
 	{
+#if WITH_EDITOR
 		using namespace UE::Interchange::Materials::Standard::Nodes;
 
 		FString MaterialUid;
@@ -527,6 +531,7 @@ namespace UE::DatasmithInterchange::MaterialUtils
 				ConnectExpression(ExpressionInput->GetExpression(), InputConnectionData, ExpressionInput->GetOutputIndex());
 			}
 		}
+#endif
 	}
 
 	void FPbrMaterialHelper::ConnectExpression(const IDatasmithMaterialExpressionCustom& DatasmithExpression, const FConnectionData& ConnectionData, int32 OutputIndex)
@@ -636,11 +641,9 @@ namespace UE::DatasmithInterchange::MaterialUtils
 
 		return BlendMode;
 	}
-#endif
 
 	bool BuildMaterialNode(IDatasmithUEPbrMaterialElement& MaterialElement, UInterchangeDatasmithPbrMaterialNode& MaterialNode)
 	{
-#if WITH_EDITOR
 		using namespace UE::Interchange::Materials;
 
 		FPbrMaterialHelper NodeHelper(MaterialNode, MaterialElement);
@@ -754,9 +757,6 @@ namespace UE::DatasmithInterchange::MaterialUtils
 		}
 
 		return true;
-#else
-		return false;
-#endif
 	}
 
 	bool BuildMaterialNode(IDatasmithDecalMaterialElement& MaterialElement, UInterchangeDatasmithMaterialNode& MaterialNode)
@@ -917,7 +917,6 @@ namespace UE::DatasmithInterchange::MaterialUtils
 			}
 			else if (MaterialElement->IsA(EDatasmithElementType::UEPbrMaterial))
 			{
-#if WITH_EDITOR
 				IDatasmithUEPbrMaterialElement& PbrMaterialElement = static_cast<IDatasmithUEPbrMaterialElement&>(*MaterialElement);
 
 				UInterchangeDatasmithPbrMaterialNode* MaterialNode = NewObject<UInterchangeDatasmithPbrMaterialNode>(&NodeContainer);
@@ -932,7 +931,6 @@ namespace UE::DatasmithInterchange::MaterialUtils
 				}
 
 				BaseNode = MaterialNode;
-#endif
 			}
 
 			if (BaseNode)
