- Adding the UInterchangeMaterialInstanceNode class to handle Material Instances as a translated node, the class supports Static Switch, Scalar, Vector and Texture parameters.
- Adding the support of Static Switch parameters in the Interchange Material Factory
- Removing the UInterchangeDatasmithMaterialNode class, now it is replaced by the UInterchangeMaterialInstanceNode, the DatasmithMaterial Pipeline and Translator now uses the new class.

diff --git a/Engine/Plugins/Experimental/Enterprise/DatasmithInterchange/Source/DatasmithInterchange/Private/InterchangeDatasmithMaterialNode.cpp b/Engine/Plugins/Experimental/Enterprise/DatasmithInterchange/Source/DatasmithInterchange/Private/InterchangeDatasmithMaterialNode.cpp
--- a/Engine/Plugins/Experimental/Enterprise/DatasmithInterchange/Source/DatasmithInterchange/Private/InterchangeDatasmithMaterialNode.cpp
+++ b/Engine/Plugins/Experimental/Enterprise/DatasmithInterchange/Source/DatasmithInterchange/Private/InterchangeDatasmithMaterialNode.cpp
@@ -28,10 +28,7 @@
 #include "Materials/MaterialInstanceDynamic.h"
 #include "Misc/PackageName.h"
 #include "UObject/ObjectRedirector.h"
-
-const FName UInterchangeDatasmithMaterialNode::MaterialTypeAttrName(TEXT("Datasmith:ReferenceMaterial:MaterialType"));
-const FName UInterchangeDatasmithMaterialNode::MaterialQualityAttrName(TEXT("Datasmith:ReferenceMaterial:MaterialQuality"));
-const FName UInterchangeDatasmithMaterialNode::MaterialParentAttrName(TEXT("Datasmith:ReferenceMaterial:MaterialParent"));
+#include "InterchangeMaterialInstanceNode.h"
 
 const FName UInterchangeDatasmithPbrMaterialNode::ShadingModelAttrName(TEXT("Datasmith:UEPbrMaterial:ShadingModel"));
 const FName UInterchangeDatasmithPbrMaterialNode::BlendModeAttrName(TEXT("Datasmith:UEPbrMaterial:BlendMode"));
@@ -44,6 +41,10 @@ namespace UE::DatasmithInterchange::MaterialUtils
 	const FName MaterialFunctionPathAttrName(TEXT("Datasmith:Material:FunctionCall:MaterialFunctionPath"));
 	const FName DefaultOutputIndexAttrName(TEXT("Datasmith:MaterialExpression:DefaultOutputIndex"));
 
+	const FName MaterialTypeAttrName(TEXT("Datasmith:ReferenceMaterial:MaterialType"));
+	const FName MaterialQualityAttrName(TEXT("Datasmith:ReferenceMaterial:MaterialQuality"));
+	const FName MaterialParentAttrName(TEXT("Datasmith:ReferenceMaterial:MaterialParent"));
+
 	class FPbrMaterialHelper
 	{
 	public:
@@ -123,12 +124,12 @@ namespace UE::DatasmithInterchange::MaterialUtils
 		EBlendMode GetUEPbrImportBlendMode();
 	};
 
-	bool BuildMaterialNode(IDatasmithMaterialInstanceElement& MaterialElement, UInterchangeDatasmithMaterialNode& MaterialNode)
+	bool BuildMaterialNode(IDatasmithMaterialInstanceElement& MaterialElement, UInterchangeMaterialInstanceNode& MaterialNode)
 	{
 		using namespace UE::Interchange::Materials;
 
-		MaterialNode.SetMaterialType(MaterialElement.GetMaterialType());
-		MaterialNode.SetMaterialQuality(MaterialElement.GetQuality());
+		MaterialNode.AddInt32Attribute(MaterialTypeAttrName, int32(MaterialElement.GetMaterialType()));
+		MaterialNode.AddInt32Attribute(MaterialQualityAttrName, int32(MaterialElement.GetQuality()));
 
 
 		UInterchangeBaseNodeContainer* NodeContainer = Cast<UInterchangeBaseNodeContainer>(MaterialNode.GetOuter());
@@ -138,7 +139,7 @@ namespace UE::DatasmithInterchange::MaterialUtils
 		for (int Index = 0; Index < MaterialElement.GetPropertiesCount(); ++Index)
 		{
 			const TSharedPtr< IDatasmithKeyValueProperty > Property = MaterialElement.GetProperty(Index);
-			const FName InputValueKey = UInterchangeShaderPortsAPI::MakeInputValueKey(FDatasmithUtils::SanitizeObjectName(Property->GetName()));
+			const FString InputValueKey = FDatasmithUtils::SanitizeObjectName(Property->GetName());
 
 			switch (Property->GetPropertyType())
 			{
@@ -147,31 +148,31 @@ namespace UE::DatasmithInterchange::MaterialUtils
 					FLinearColor LinearColor;
 					LinearColor.InitFromString(Property->GetValue());
 
-					MaterialNode.AddLinearColorAttribute(InputValueKey, LinearColor);
+					MaterialNode.AddVectorParameterValue(InputValueKey, LinearColor);
 					break;
 				}
 
 				case EDatasmithKeyValuePropertyType::Float:
 				{
-					MaterialNode.AddFloatAttribute(InputValueKey, FCString::Atof(Property->GetValue()));
+					MaterialNode.AddScalarParameterValue(InputValueKey, FCString::Atof(Property->GetValue()));
 					break;
 				}
 
 				case EDatasmithKeyValuePropertyType::Bool:
 				{
-					MaterialNode.AddBooleanAttribute(InputValueKey, FString(Property->GetValue()).ToBool());
+					MaterialNode.AddStaticSwitchParameterValue(InputValueKey, FString(Property->GetValue()).ToBool());
 					break;
 				}
 
 				case EDatasmithKeyValuePropertyType::Integer:
 				{
-					MaterialNode.AddInt32Attribute(InputValueKey, FCString::Atoi(Property->GetValue()));
+					MaterialNode.AddInt32Attribute(UInterchangeShaderPortsAPI::MakeInputValueKey(InputValueKey), FCString::Atoi(Property->GetValue()));
 					break;
 				}
 
 				case EDatasmithKeyValuePropertyType::Texture:
 				{
-					MaterialNode.AddStringAttribute(InputValueKey, Property->GetValue());
+					MaterialNode.AddTextureParameterValue(InputValueKey, Property->GetValue());
 					break;
 				}
 
@@ -187,14 +188,6 @@ namespace UE::DatasmithInterchange::MaterialUtils
 		return true;
 	}
 
-	bool BuildMaterialNode(IDatasmithMaterialElement& MaterialElement, UInterchangeDatasmithMaterialNode& MaterialNode)
-	{
-		using namespace UE::Interchange::Materials;
-
-		return false;
-	}
-
-
 	FPbrMaterialHelper::FPbrMaterialHelper(UInterchangeShaderNode& MaterialNode, IDatasmithUEPbrMaterialElement& InMaterialElement)
 		: NodeContainer(*Cast<UInterchangeBaseNodeContainer>(MaterialNode.GetOuter()))
 		, MaterialID(MaterialNode.GetUniqueID())
@@ -760,13 +753,6 @@ namespace UE::DatasmithInterchange::MaterialUtils
 		return true;
 	}
 
-	bool BuildMaterialNode(IDatasmithDecalMaterialElement& MaterialElement, UInterchangeDatasmithMaterialNode& MaterialNode)
-	{
-		using namespace UE::Interchange::Materials;
-
-		return false;
-	}
-
 	void ProcessDependencies(IDatasmithUEPbrMaterialElement& MaterialElement, TArray<TSharedPtr<IDatasmithBaseMaterialElement>>& OutMaterialElements, TMap<FString, TSharedPtr<IDatasmithBaseMaterialElement>>& MaterialsToSort)
 	{
 		for (int32 ExpressionIndex = 0; ExpressionIndex < MaterialElement.GetExpressionsCount(); ExpressionIndex++)
@@ -905,7 +891,7 @@ namespace UE::DatasmithInterchange::MaterialUtils
 		{
 			if (MaterialElement->IsA(EDatasmithElementType::MaterialInstance))
 			{
-				UInterchangeDatasmithMaterialNode* MaterialNode = NewObject< UInterchangeDatasmithMaterialNode >(&NodeContainer);
+				UInterchangeMaterialInstanceNode* MaterialNode = NewObject<UInterchangeMaterialInstanceNode>(&NodeContainer);
 				const FString MaterialNodeUid = NodeUtils::MaterialPrefix + FDatasmithUtils::SanitizeObjectName(MaterialElement->GetName());
 				MaterialNode->InitializeNode(MaterialNodeUid, *FDatasmithUtils::SanitizeObjectName(MaterialElement->GetLabel()), EInterchangeNodeContainerType::TranslatedAsset);
 
diff --git a/Engine/Plugins/Experimental/Enterprise/DatasmithInterchange/Source/DatasmithInterchange/Private/InterchangeDatasmithMaterialNode.h b/Engine/Plugins/Experimental/Enterprise/DatasmithInterchange/Source/DatasmithInterchange/Private/InterchangeDatasmithMaterialNode.h
--- a/Engine/Plugins/Experimental/Enterprise/DatasmithInterchange/Source/DatasmithInterchange/Private/InterchangeDatasmithMaterialNode.h
+++ b/Engine/Plugins/Experimental/Enterprise/DatasmithInterchange/Source/DatasmithInterchange/Private/InterchangeDatasmithMaterialNode.h
@@ -17,61 +17,6 @@
 class IDatasmithBaseMaterialElement;
 class UInterchangeBaseNodeContainer;
 
-UCLASS(BlueprintType, Experimental)
-class DATASMITHINTERCHANGE_API UInterchangeDatasmithMaterialNode : public UInterchangeBaseNode
-{
-	GENERATED_BODY()
-
-public:
-	/**
-	 * Return the node type name of the class, we use this when reporting error
-	 */
-	virtual FString GetTypeName() const override
-	{
-		const FString TypeName = TEXT("DatasmithMaterialNode");
-		return TypeName;
-	}
-
-	void SetMaterialType(EDatasmithReferenceMaterialType MaterialType)
-	{
-		AddInt32Attribute(MaterialTypeAttrName, (int32)MaterialType);
-	}
-
-	EDatasmithReferenceMaterialType GetMaterialType()
-	{
-		int32 AttributeValue;
-		return GetInt32Attribute(MaterialTypeAttrName, AttributeValue) ? (EDatasmithReferenceMaterialType)AttributeValue : EDatasmithReferenceMaterialType::Opaque;
-	}
-
-	void SetMaterialQuality(EDatasmithReferenceMaterialQuality MaterialQuality)
-	{
-		AddInt32Attribute(MaterialQualityAttrName, (int32)MaterialQuality);
-	}
-
-	EDatasmithReferenceMaterialQuality GetMaterialQuality()
-	{
-		int32 AttributeValue;
-		return GetInt32Attribute(MaterialQualityAttrName, AttributeValue) ? (EDatasmithReferenceMaterialQuality)AttributeValue : EDatasmithReferenceMaterialQuality::High;
-	}
-
-	void SetParentPath(const TCHAR* InSelectorName)
-	{
-		AddStringAttribute(MaterialParentAttrName, InSelectorName ? InSelectorName : FString());
-	}
-
-	FString GetParentPath()
-	{
-		FString AttributeValue;
-		return GetStringAttribute(MaterialParentAttrName, AttributeValue) ? AttributeValue : FString();
-	}
-
-public:
-	// TODO: See if it is not more appropriate to use IMPLEMENT_NODE_ATTRIBUTE_GETTER/SETTER defines
-	static const FName MaterialTypeAttrName;
-	static const FName MaterialQualityAttrName;
-	static const FName MaterialParentAttrName;
-};
-
 UCLASS(BlueprintType, Experimental)
 class DATASMITHINTERCHANGE_API UInterchangeDatasmithPbrMaterialNode : public UInterchangeShaderGraphNode
 {
@@ -184,6 +129,10 @@ namespace UE::DatasmithInterchange::MaterialUtils
 	extern const FName DefaultOutputIndexAttrName;
 	extern const FName MaterialFunctionPathAttrName;
 
+	extern const FName MaterialTypeAttrName;
+	extern const FName MaterialQualityAttrName;
+	extern const FName MaterialParentAttrName;
+
 	void ProcessMaterialElements(TArray<TSharedPtr<IDatasmithBaseMaterialElement>>& InOutMaterialElements);
 	UInterchangeBaseNode* AddMaterialNode(const TSharedPtr<IDatasmithBaseMaterialElement>& MaterialElement, UInterchangeBaseNodeContainer& NodeContainer);
 }
diff --git a/Engine/Plugins/Experimental/Enterprise/DatasmithInterchange/Source/DatasmithInterchange/Private/InterchangeDatasmithMaterialPipeline.cpp b/Engine/Plugins/Experimental/Enterprise/DatasmithInterchange/Source/DatasmithInterchange/Private/InterchangeDatasmithMaterialPipeline.cpp
--- a/Engine/Plugins/Experimental/Enterprise/DatasmithInterchange/Source/DatasmithInterchange/Private/InterchangeDatasmithMaterialPipeline.cpp
+++ b/Engine/Plugins/Experimental/Enterprise/DatasmithInterchange/Source/DatasmithInterchange/Private/InterchangeDatasmithMaterialPipeline.cpp
@@ -17,6 +17,7 @@
 #include "InterchangeMaterialFactoryNode.h"
 #include "InterchangeTexture2DNode.h"
 #include "InterchangeTexture2DFactoryNode.h"
+#include "InterchangeMaterialInstanceNode.h"
 
 #include "Materials/MaterialFunction.h"
 #include "Materials/MaterialInstanceConstant.h"
@@ -33,13 +34,13 @@ void UInterchangeDatasmithMaterialPipeline::ExecutePreImportPipeline(UInterchang
 {
 	using namespace UE::DatasmithInterchange;
 
-	TArray< UInterchangeDatasmithMaterialNode*> InstancedMaterials;
+	TArray<UInterchangeMaterialInstanceNode*> InstancedMaterials;
 	TArray<UInterchangeShaderNode*> ShaderNodes;
 
 	//Find all translated node we need for this pipeline
 	NodeContainer->IterateNodes([&](const FString& NodeUid, UInterchangeBaseNode* Node)
 		{
-			if (UInterchangeDatasmithMaterialNode* MaterialNode = Cast<UInterchangeDatasmithMaterialNode>(Node))
+			if (UInterchangeMaterialInstanceNode* MaterialNode = Cast<UInterchangeMaterialInstanceNode>(Node))
 			{
 				InstancedMaterials.Add(MaterialNode);
 			}
@@ -53,7 +54,7 @@ void UInterchangeDatasmithMaterialPipeline::ExecutePreImportPipeline(UInterchang
 
 	UpdateMaterialFactoryNodes(ShaderNodes);
 
-	for (UInterchangeDatasmithMaterialNode* MaterialNode : InstancedMaterials)
+	for (UInterchangeMaterialInstanceNode* MaterialNode : InstancedMaterials)
 	{
 		PreImportMaterialNode(NodeContainer, MaterialNode);
 	}
@@ -86,11 +87,11 @@ void UInterchangeDatasmithMaterialPipeline::ExecutePostImportPipeline(const UInt
 	}
 }
 
-void UInterchangeDatasmithMaterialPipeline::PreImportMaterialNode(UInterchangeBaseNodeContainer* NodeContainer, UInterchangeDatasmithMaterialNode* MaterialNode)
+void UInterchangeDatasmithMaterialPipeline::PreImportMaterialNode(UInterchangeBaseNodeContainer* NodeContainer, UInterchangeMaterialInstanceNode* MaterialNode)
 {
-	FString ParentPath = MaterialNode->GetParentPath();
+	FString ParentPath;	
 
-	if (ParentPath.IsEmpty())
+	if(!MaterialNode->GetCustomParent(ParentPath) || ParentPath.IsEmpty())
 	{
 		return;
 	}
@@ -99,22 +100,26 @@ void UInterchangeDatasmithMaterialPipeline::PreImportMaterialNode(UInterchangeBa
 
 	FString PackagePath;
 
-	if (MaterialNode->GetMaterialType() == EDatasmithReferenceMaterialType::Custom && FPackageName::DoesPackageExist(ParentPath))
+	int32 MaterialType;
+	if (MaterialNode->GetInt32Attribute(MaterialUtils::MaterialTypeAttrName, MaterialType))
 	{
-		PackagePath = ParentPath;
-	}
-	else
-	{
-		FString Host(FDatasmithReferenceMaterialManager::Get().GetHostFromString(ParentPath));
-
-		if (TSharedPtr<FDatasmithReferenceMaterialSelector> MaterialSelector = FDatasmithReferenceMaterialManager::Get().GetSelector(*Host))
+		if (EDatasmithReferenceMaterialType(MaterialType) == EDatasmithReferenceMaterialType::Custom && FPackageName::DoesPackageExist(ParentPath))
 		{
-			PackagePath = MaterialSelector->GetMaterialPath(MaterialNode->GetMaterialType());
+			PackagePath = ParentPath;
 		}
 		else
 		{
-			//FText FailReason = FText::Format(LOCTEXT("NoSelectorForHost", "No Material selector found for Host {0}. Skipping material {1} ..."), FText::FromString(Host), FText::FromString(MaterialElement.GetName()));
-			//ImportContext.LogError(FailReason);
+			FString Host(FDatasmithReferenceMaterialManager::Get().GetHostFromString(ParentPath));
+
+			if (TSharedPtr<FDatasmithReferenceMaterialSelector> MaterialSelector = FDatasmithReferenceMaterialManager::Get().GetSelector(*Host))
+			{
+				PackagePath = MaterialSelector->GetMaterialPath(EDatasmithReferenceMaterialType(MaterialType));
+			}
+			else
+			{
+				//FText FailReason = FText::Format(LOCTEXT("NoSelectorForHost", "No Material selector found for Host {0}. Skipping material {1} ..."), FText::FromString(Host), FText::FromString(MaterialElement.GetName()));
+				//ImportContext.LogError(FailReason);
+			}
 		}
 	}
 
@@ -196,11 +201,14 @@ void UInterchangeDatasmithMaterialPipeline::PreImportMaterialNode(UInterchangeBa
 		}
 	}
 
-	if (MaterialNode->GetMaterialType() != EDatasmithReferenceMaterialType::Custom)
+	if (EDatasmithReferenceMaterialType(MaterialType) != EDatasmithReferenceMaterialType::Custom)
 	{
-		MaterialFactoryNode->AddStringAttribute(UInterchangeDatasmithMaterialNode::MaterialParentAttrName, MaterialNode->GetParentPath());
-		MaterialFactoryNode->AddInt32Attribute(UInterchangeDatasmithMaterialNode::MaterialTypeAttrName, (int32)MaterialNode->GetMaterialType());
-		MaterialFactoryNode->AddInt32Attribute(UInterchangeDatasmithMaterialNode::MaterialQualityAttrName, (int32)MaterialNode->GetMaterialQuality());
+		MaterialFactoryNode->AddStringAttribute(MaterialUtils::MaterialParentAttrName, ParentPath);
+		MaterialFactoryNode->AddInt32Attribute(MaterialUtils::MaterialTypeAttrName, MaterialType);
+		if(int32 MaterialQuality; MaterialNode->GetInt32Attribute(MaterialUtils::MaterialQualityAttrName, MaterialQuality))
+		{
+			MaterialFactoryNode->AddInt32Attribute(MaterialUtils::MaterialQualityAttrName, MaterialQuality);
+		}
 	}
 }
 
@@ -210,7 +218,7 @@ void UInterchangeDatasmithMaterialPipeline::PostImportMaterialInstanceFactoryNod
 
 #if WITH_EDITOR
 	FString SelectorName;
-	if (FactoryNode->GetStringAttribute(UInterchangeDatasmithMaterialNode::MaterialParentAttrName, SelectorName))
+	if (FactoryNode->GetStringAttribute(MaterialUtils::MaterialParentAttrName, SelectorName))
 	{
 		if (!SelectorName.IsEmpty())
 		{
@@ -219,8 +227,8 @@ void UInterchangeDatasmithMaterialPipeline::PostImportMaterialInstanceFactoryNod
 			int32 MaterialType;
 			int32 MaterialQuality;
 
-			ensure(FactoryNode->GetInt32Attribute(UInterchangeDatasmithMaterialNode::MaterialTypeAttrName, MaterialType));
-			ensure(FactoryNode->GetInt32Attribute(UInterchangeDatasmithMaterialNode::MaterialQualityAttrName, MaterialQuality));
+			ensure(FactoryNode->GetInt32Attribute(MaterialUtils::MaterialTypeAttrName, MaterialType));
+			ensure(FactoryNode->GetInt32Attribute(MaterialUtils::MaterialQualityAttrName, MaterialQuality));
 
 			if (TSharedPtr<FDatasmithReferenceMaterialSelector> MaterialSelector = FDatasmithReferenceMaterialManager::Get().GetSelector(*Host))
 			{
diff --git a/Engine/Plugins/Experimental/Enterprise/DatasmithInterchange/Source/DatasmithInterchange/Private/InterchangeDatasmithMaterialPipeline.h b/Engine/Plugins/Experimental/Enterprise/DatasmithInterchange/Source/DatasmithInterchange/Private/InterchangeDatasmithMaterialPipeline.h
--- a/Engine/Plugins/Experimental/Enterprise/DatasmithInterchange/Source/DatasmithInterchange/Private/InterchangeDatasmithMaterialPipeline.h
+++ b/Engine/Plugins/Experimental/Enterprise/DatasmithInterchange/Source/DatasmithInterchange/Private/InterchangeDatasmithMaterialPipeline.h
@@ -9,12 +9,12 @@
 #include "InterchangeDatasmithMaterialPipeline.generated.h"
 
 class UMaterialInterface;
-class UInterchangeDatasmithMaterialNode;
 class UInterchangeDatasmithPbrMaterialNode;
 class UInterchangeFactoryBaseNode;
 class UInterchangeMaterialFactoryNode;
 class UInterchangeMaterialInstanceFactoryNode;
 class UInterchangeShaderGraphNode;
+class UInterchangeMaterialInstanceNode;
 
 UCLASS(BlueprintType, Experimental)
 class DATASMITHINTERCHANGE_API UInterchangeDatasmithMaterialPipeline : public UInterchangeGenericMaterialPipeline
@@ -34,7 +34,7 @@ protected:
 	}
 
 private:
-	void PreImportMaterialNode(UInterchangeBaseNodeContainer* NodeContainer, UInterchangeDatasmithMaterialNode* MaterialNode);
+	void PreImportMaterialNode(UInterchangeBaseNodeContainer* NodeContainer, UInterchangeMaterialInstanceNode* MaterialNode);
 	void PostImportMaterialInstanceFactoryNode(const UInterchangeBaseNodeContainer* NodeContainer, UInterchangeMaterialInstanceFactoryNode* FactoryNode, UMaterialInterface* CreatedMaterial);
 	void UpdateMaterialFactoryNodes(const TArray<UInterchangeShaderNode*>& ShaderNodes);
 };
\ No newline at end of file
diff --git a/Engine/Plugins/Experimental/Enterprise/DatasmithInterchange/Source/DatasmithInterchange/Private/InterchangeDatasmithTranslator.cpp b/Engine/Plugins/Experimental/Enterprise/DatasmithInterchange/Source/DatasmithInterchange/Private/InterchangeDatasmithTranslator.cpp
--- a/Engine/Plugins/Experimental/Enterprise/DatasmithInterchange/Source/DatasmithInterchange/Private/InterchangeDatasmithTranslator.cpp
+++ b/Engine/Plugins/Experimental/Enterprise/DatasmithInterchange/Source/DatasmithInterchange/Private/InterchangeDatasmithTranslator.cpp
@@ -26,6 +26,7 @@
 #include "InterchangeLightNode.h"
 #include "InterchangeManager.h"
 #include "InterchangeMaterialDefinitions.h"
+#include "InterchangeMaterialInstanceNode.h"
 #include "InterchangeMeshNode.h"
 #include "InterchangeShaderGraphNode.h"
 #include "InterchangeSceneNode.h"
@@ -172,14 +173,14 @@ bool UInterchangeDatasmithTranslator::Translate(UInterchangeBaseNodeContainer& B
 
 				if (MaterialElement->IsA(EDatasmithElementType::MaterialInstance))
 				{
-					UInterchangeDatasmithMaterialNode* ReferenceMaterialNode = Cast<UInterchangeDatasmithMaterialNode>(MaterialNode);
-					if (ReferenceMaterialNode->GetMaterialType() == EDatasmithReferenceMaterialType::Custom)
+					UInterchangeMaterialInstanceNode* ReferenceMaterialNode = Cast<UInterchangeMaterialInstanceNode>(MaterialNode);
+					if (int32 MaterialType; ReferenceMaterialNode->GetInt32Attribute(MaterialUtils::MaterialTypeAttrName, MaterialType) && EDatasmithReferenceMaterialType(MaterialType) == EDatasmithReferenceMaterialType::Custom)
 					{
-						ReferenceMaterialNode->SetParentPath(static_cast<IDatasmithMaterialInstanceElement&>(*MaterialElement).GetCustomMaterialPathName());
+						ReferenceMaterialNode->SetCustomParent(static_cast<IDatasmithMaterialInstanceElement&>(*MaterialElement).GetCustomMaterialPathName());
 					}
 					else
 					{
-						ReferenceMaterialNode->SetParentPath(HostName);
+						ReferenceMaterialNode->SetCustomParent(HostName);
 					}
 				}
 			}
diff --git a/Engine/Plugins/Experimental/Enterprise/DatasmithInterchange/Source/DatasmithInterchange/Private/InterchangeDatasmithUtils.h b/Engine/Plugins/Experimental/Enterprise/DatasmithInterchange/Source/DatasmithInterchange/Private/InterchangeDatasmithUtils.h
--- a/Engine/Plugins/Experimental/Enterprise/DatasmithInterchange/Source/DatasmithInterchange/Private/InterchangeDatasmithUtils.h
+++ b/Engine/Plugins/Experimental/Enterprise/DatasmithInterchange/Source/DatasmithInterchange/Private/InterchangeDatasmithUtils.h
@@ -14,7 +14,6 @@ class IDatasmithBaseAnimationElement;
 class IDatasmithBaseMaterialElement;
 class IDatasmithLevelVariantSetsElement;
 class IDatasmithVariantSetElement;
-class UInterchangeDatasmithMaterialNode;
 
 namespace UE::Interchange
 {
diff --git a/Engine/Plugins/Interchange/Runtime/Source/Import/Private/Material/InterchangeMaterialFactory.cpp b/Engine/Plugins/Interchange/Runtime/Source/Import/Private/Material/InterchangeMaterialFactory.cpp
--- a/Engine/Plugins/Interchange/Runtime/Source/Import/Private/Material/InterchangeMaterialFactory.cpp
+++ b/Engine/Plugins/Interchange/Runtime/Source/Import/Private/Material/InterchangeMaterialFactory.cpp
@@ -1012,6 +1012,27 @@ void UInterchangeMaterialFactory::SetupMaterialInstance(UMaterialInstance* Mater
 
 		switch(UInterchangeShaderPortsAPI::GetInputType(MaterialFactoryNode, InputName))
 		{
+#if WITH_EDITORONLY_DATA
+		case UE::Interchange::EAttributeTypes::Bool:
+		    {
+				bool bInstanceValue;
+				FGuid Uid;
+				if(MaterialInstance->GetStaticSwitchParameterValue(ParameterName, bInstanceValue, Uid))
+				{
+					bool bInputValue = false;
+					MaterialFactoryNode->GetBooleanAttribute(UInterchangeShaderPortsAPI::MakeInputValueKey(InputName), bInputValue);
+
+					if(bInputValue != bInstanceValue)
+					{
+						if(UMaterialInstanceConstant* MaterialInstanceConstant = Cast<UMaterialInstanceConstant>(MaterialInstance))
+						{
+							MaterialInstanceConstant->SetStaticSwitchParameterValueEditorOnly(ParameterName, bInputValue);
+						}
+					}
+				}
+		    }
+		    break;
+#endif // #if WITH_EDITORONLY_DATA
 		case UE::Interchange::EAttributeTypes::Float:
 			{
 				float InstanceValue;
diff --git a/Engine/Plugins/Interchange/Runtime/Source/Nodes/Private/InterchangeMaterialInstanceNode.cpp b/Engine/Plugins/Interchange/Runtime/Source/Nodes/Private/InterchangeMaterialInstanceNode.cpp
new file mode 100644
--- /dev/null
+++ b/Engine/Plugins/Interchange/Runtime/Source/Nodes/Private/InterchangeMaterialInstanceNode.cpp
@@ -0,0 +1,93 @@
+// Copyright Epic Games, Inc. All Rights Reserved.
+
+#include "InterchangeMaterialInstanceNode.h"
+#include "Nodes/InterchangeBaseNodeContainer.h"
+#include "InterchangeShaderGraphNode.h"
+
+#include UE_INLINE_GENERATED_CPP_BY_NAME(InterchangeMaterialInstanceNode)
+
+FString UInterchangeMaterialInstanceNode::GetTypeName() const
+{
+    return TEXT("MaterialInstanceNode");
+}
+
+FString UInterchangeMaterialInstanceNode::MakeNodeUid(const FStringView NodeName, const FStringView ParentNodeUid)
+{
+	FString MaterialInstanceNodeUid = FString(UInterchangeBaseNode::HierarchySeparator) + TEXT("MaterialInstances") + FString(UInterchangeBaseNode::HierarchySeparator);
+
+	if(ParentNodeUid.IsEmpty())
+	{
+		MaterialInstanceNodeUid += NodeName;
+	}
+	else
+	{
+		MaterialInstanceNodeUid += ParentNodeUid + FString(UInterchangeBaseNode::HierarchySeparator) + NodeName;
+	}
+
+	return MaterialInstanceNodeUid;
+}
+
+UInterchangeMaterialInstanceNode* UInterchangeMaterialInstanceNode::Create(UInterchangeBaseNodeContainer* NodeContainer, const FStringView NodeName, const FStringView ParentNodeUid)
+{
+	check(NodeContainer);
+
+	const FString MaterialInstanceNodeUid = MakeNodeUid(NodeName, ParentNodeUid);
+
+	UInterchangeMaterialInstanceNode* MaterialInstanceNode = NewObject<UInterchangeMaterialInstanceNode>(NodeContainer);
+	MaterialInstanceNode->InitializeNode(MaterialInstanceNodeUid, FString(NodeName), EInterchangeNodeContainerType::TranslatedAsset);
+
+	NodeContainer->AddNode(MaterialInstanceNode);
+	NodeContainer->SetNodeParentUid(MaterialInstanceNodeUid, FString(ParentNodeUid));
+
+	return MaterialInstanceNode;
+}
+
+bool UInterchangeMaterialInstanceNode::SetCustomParent(const FString& AttributeValue) const
+{
+	IMPLEMENT_NODE_ATTRIBUTE_SETTER_NODELEGATE(Parent, FString);
+}
+
+bool UInterchangeMaterialInstanceNode::GetCustomParent(FString& AttributeValue) const
+{
+	IMPLEMENT_NODE_ATTRIBUTE_GETTER(Parent, FString);
+}
+
+bool UInterchangeMaterialInstanceNode::AddScalarParameterValue(const FString& ParameterName, float AttributeValue)
+{
+	return AddFloatAttribute(UInterchangeShaderPortsAPI::MakeInputValueKey(ParameterName), AttributeValue);
+}
+
+bool UInterchangeMaterialInstanceNode::GetScalarParameterValue(const FString& ParameterName, float& AttributeValue) const
+{
+	return GetFloatAttribute(UInterchangeShaderPortsAPI::MakeInputValueKey(ParameterName), AttributeValue);
+}
+
+bool UInterchangeMaterialInstanceNode::AddVectorParameterValue(const FString& ParameterName, const FLinearColor& AttributeValue)
+{
+	return AddLinearColorAttribute(UInterchangeShaderPortsAPI::MakeInputValueKey(ParameterName), AttributeValue);
+}
+
+bool UInterchangeMaterialInstanceNode::GetVectorParameterValue(const FString& ParameterName, FLinearColor& AttributeValue) const
+{
+	return GetLinearColorAttribute(UInterchangeShaderPortsAPI::MakeInputValueKey(ParameterName), AttributeValue);
+}
+
+bool UInterchangeMaterialInstanceNode::AddTextureParameterValue(const FString& ParameterName, const FString& AttributeValue)
+{
+	return AddStringAttribute(UInterchangeShaderPortsAPI::MakeInputValueKey(ParameterName), AttributeValue);
+}
+
+bool UInterchangeMaterialInstanceNode::GetTextureParameterValue(const FString& ParameterName, FString& AttributeValue) const
+{
+	return GetStringAttribute(UInterchangeShaderPortsAPI::MakeInputValueKey(ParameterName), AttributeValue);
+}
+
+bool UInterchangeMaterialInstanceNode::AddStaticSwitchParameterValue(const FString& ParameterName, bool AttributeValue)
+{
+	return AddBooleanAttribute(UInterchangeShaderPortsAPI::MakeInputValueKey(ParameterName), AttributeValue);
+}
+
+bool UInterchangeMaterialInstanceNode::GetStaticSwitchParameterValue(const FString& ParameterName, bool& AttributeValue) const
+{
+	return GetBooleanAttribute(UInterchangeShaderPortsAPI::MakeInputValueKey(ParameterName), AttributeValue);
+}
diff --git a/Engine/Plugins/Interchange/Runtime/Source/Nodes/Public/InterchangeMaterialInstanceNode.h b/Engine/Plugins/Interchange/Runtime/Source/Nodes/Public/InterchangeMaterialInstanceNode.h
new file mode 100644
--- /dev/null
+++ b/Engine/Plugins/Interchange/Runtime/Source/Nodes/Public/InterchangeMaterialInstanceNode.h
@@ -0,0 +1,69 @@
+// Copyright Epic Games, Inc. All Rights Reserved.
+
+#pragma once
+
+#include "CoreMinimal.h"
+#include "Nodes/InterchangeBaseNode.h"
+
+#include "InterchangeMaterialInstanceNode.generated.h"
+
+class UInterchangeBaseNodeContainer;
+
+UCLASS(BlueprintType)
+class INTERCHANGENODES_API UInterchangeMaterialInstanceNode : public UInterchangeBaseNode
+{
+	GENERATED_BODY()
+
+public:
+
+	/**
+     * Return the node type name of the class, we use this when reporting error
+     */
+	virtual FString GetTypeName() const override;
+
+
+	/**
+	 * Build and return a UID name for a shader node.
+	 */
+	static FString MakeNodeUid(const FStringView NodeName, const FStringView ParentNodeUid);
+
+	/**
+     * Creates a new UInterchangeShaderGraphNode and adds it to NodeContainer as a translated node.
+     */
+	static UInterchangeMaterialInstanceNode* Create(UInterchangeBaseNodeContainer* NodeContainer, const FStringView NodeName, const FStringView ParentNodeUid);
+
+	UFUNCTION(BlueprintCallable, Category = "Interchange | Node | Material Instance")
+	bool SetCustomParent(const FString& AttributeValue) const;
+
+	UFUNCTION(BlueprintCallable, Category = "Interchange | Node | Material Instance")
+	bool GetCustomParent(FString& AttributeValue) const;
+
+	UFUNCTION(BlueprintCallable, Category = "Interchange | Node | Material Instance")
+	bool AddScalarParameterValue(const FString & ParameterName, float AttributeValue);
+
+	UFUNCTION(BlueprintCallable, Category = "Interchange | Node | Material Instance")
+	bool GetScalarParameterValue(const FString& ParameterName, float& AttributeValue) const;
+
+	UFUNCTION(BlueprintCallable, Category = "Interchange | Node | Material Instance")
+	bool AddVectorParameterValue(const FString& ParameterName, const FLinearColor& AttributeValue);
+
+	UFUNCTION(BlueprintCallable, Category = "Interchange | Node | Material Instance")
+	bool GetVectorParameterValue(const FString& ParameterName, FLinearColor& AttributeValue) const;
+
+	UFUNCTION(BlueprintCallable, Category = "Interchange | Node | Material Instance")
+	bool AddTextureParameterValue(const FString& ParameterName, const FString& AttributeValue);
+
+	UFUNCTION(BlueprintCallable, Category = "Interchange | Node | Material Instance")
+	bool GetTextureParameterValue(const FString& ParameterName, FString& AttributeValue) const;
+
+	UFUNCTION(BlueprintCallable, Category = "Interchange | Node | Material Instance")
+	bool AddStaticSwitchParameterValue(const FString& ParameterName, bool AttributeValue);
+
+	UFUNCTION(BlueprintCallable, Category = "Interchange | Node | Material Instance")
+	bool GetStaticSwitchParameterValue(const FString& ParameterName, bool& AttributeValue) const;
+
+private:
+
+	IMPLEMENT_NODE_ATTRIBUTE_KEY(Parent);
+
+};
\ No newline at end of file
