Fix CIS build

diff --git a/Engine/Plugins/Enterprise/DatasmithImporter/Source/DatasmithExternalSource/Private/DatasmithFileUriResolver.cpp b/Engine/Plugins/Enterprise/DatasmithImporter/Source/DatasmithExternalSource/Private/DatasmithFileUriResolver.cpp
--- a/Engine/Plugins/Enterprise/DatasmithImporter/Source/DatasmithExternalSource/Private/DatasmithFileUriResolver.cpp
+++ b/Engine/Plugins/Enterprise/DatasmithImporter/Source/DatasmithExternalSource/Private/DatasmithFileUriResolver.cpp
@@ -7,7 +7,6 @@
 #include "DatasmithTranslatorManager.h"
 
 #if WITH_EDITOR
-#include "AssetToolsModule.h"
 #include "DesktopPlatformModule.h"
 #include "Framework/Application/SlateApplication.h"
 #include "IDesktopPlatform.h"
@@ -50,8 +49,7 @@ namespace UE::DatasmithImporter
 		FString FileTypes;
 		FString Extensions;
 
-		IAssetTools& AssetTools = FModuleManager::LoadModuleChecked<FAssetToolsModule>(TEXT("AssetTools")).Get();
-		ObjectTools::AppendFormatsFileExtensions(Formats, FileTypes, Extensions, AssetTools.GetSupportedImportExtension());
+		ObjectTools::AppendFormatsFileExtensions(Formats, FileTypes, Extensions);
 
 		const FString FormatString = FString::Printf(TEXT("All Files (%s)|%s|%s"), *Extensions, *Extensions, *FileTypes);
 		const FString Title = NSLOCTEXT("DatasmithFileUriResolver", "BrowseSourceDialogTitle", "Select Datasmith Source File").ToString();
diff --git a/Engine/Plugins/Enterprise/DatasmithImporter/Source/DatasmithImporter/Private/DatasmithImporterHelper.cpp b/Engine/Plugins/Enterprise/DatasmithImporter/Source/DatasmithImporter/Private/DatasmithImporterHelper.cpp
--- a/Engine/Plugins/Enterprise/DatasmithImporter/Source/DatasmithImporter/Private/DatasmithImporterHelper.cpp
+++ b/Engine/Plugins/Enterprise/DatasmithImporter/Source/DatasmithImporter/Private/DatasmithImporterHelper.cpp
@@ -5,7 +5,6 @@
 #include "Utility/DatasmithImporterUtils.h"
 
 #include "AssetRegistry/AssetRegistryModule.h"
-#include "AssetToolsModule.h"
 #include "BusyCursor.h"
 #include "ContentBrowserModule.h"
 #include "DesktopPlatformModule.h"
@@ -254,8 +253,7 @@ FString FDatasmithImporterHelper::GetFilterStringInternal(UFactory* Factory)
 	FString AllExtensions;
 	TMultiMap<uint32, UFactory*> FilterIndexToFactory;
 
-	IAssetTools& AssetTools = FModuleManager::LoadModuleChecked<FAssetToolsModule>(TEXT("AssetTools")).Get();
-	ObjectTools::GenerateFactoryFileExtensions(Factories, FileTypes, AllExtensions, FilterIndexToFactory, AssetTools.GetSupportedImportExtension());
+	ObjectTools::GenerateFactoryFileExtensions(Factories, FileTypes, AllExtensions, FilterIndexToFactory);
 
 	return FString::Printf(TEXT("All Files (%s)|%s|%s"), *AllExtensions, *AllExtensions, *FileTypes);
 }
diff --git a/Engine/Plugins/Experimental/Enterprise/DatasmithInterchange/Source/DatasmithInterchange/Private/DatasmithInterchangeModule.cpp b/Engine/Plugins/Experimental/Enterprise/DatasmithInterchange/Source/DatasmithInterchange/Private/DatasmithInterchangeModule.cpp
--- a/Engine/Plugins/Experimental/Enterprise/DatasmithInterchange/Source/DatasmithInterchange/Private/DatasmithInterchangeModule.cpp
+++ b/Engine/Plugins/Experimental/Enterprise/DatasmithInterchange/Source/DatasmithInterchange/Private/DatasmithInterchangeModule.cpp
@@ -24,7 +24,6 @@
 #include "Logging/LogMacros.h"
 
 #if WITH_EDITOR
-#include "AssetToolsModule.h"
 #include "DatasmithImporterModule.h"
 #include "DatasmithUtils.h"
 #include "DesktopPlatformModule.h"
@@ -133,8 +132,7 @@ void FDatasmithInterchangeModule::OnImportInterchange()
 	FString Extensions;
 
 	FString DefaultLocation(FEditorDirectories::Get().GetLastDirectory(ELastDirectory::GENERIC_IMPORT));
-	IAssetTools& AssetTools = FModuleManager::LoadModuleChecked<FAssetToolsModule>(TEXT("AssetTools")).Get();
-	ObjectTools::AppendFormatsFileExtensions(Formats, FileTypes, Extensions, AssetTools.GetSupportedImportExtension());
+	ObjectTools::AppendFormatsFileExtensions(Formats, FileTypes, Extensions);
 
 	const FString FormatString = FString::Printf(TEXT("All Files (%s)|%s|%s"), *Extensions, *Extensions, *FileTypes);
 	const FString Title = LOCTEXT("BrowseSourceDialogTitle", "Select Datasmith Source File").ToString();
diff --git a/Engine/Source/Runtime/Interchange/Engine/Private/InterchangeManager.cpp b/Engine/Source/Runtime/Interchange/Engine/Private/InterchangeManager.cpp
--- a/Engine/Source/Runtime/Interchange/Engine/Private/InterchangeManager.cpp
+++ b/Engine/Source/Runtime/Interchange/Engine/Private/InterchangeManager.cpp
@@ -5,7 +5,9 @@
 #include "AssetRegistry/AssetData.h"
 #include "AssetRegistry/AssetDataTagMap.h"
 #include "AssetRegistry/AssetRegistryModule.h"
+#if WITH_EDITOR
 #include "AssetToolsModule.h"
+#endif
 #include "Async/Async.h"
 #include "CoreGlobals.h"
 #include "CoreMinimal.h"
@@ -808,13 +810,14 @@ bool UInterchangeManager::CanTranslateSourceData(const UInterchangeSourceData* S
 	{
 		return false;
 	}
-	
+#if WITH_EDITOR
 	IAssetTools& AssetTools = FModuleManager::LoadModuleChecked<FAssetToolsModule>(TEXT("AssetTools")).Get();
 	const FString Extension = FPaths::GetExtension(SourceData->GetFilename());
 	if (!AssetTools.IsImportExtensionAllowed(Extension))
 	{
 		return false;
 	}
+#endif
 
 	return GetTranslatorForSourceData(SourceData) != nullptr;
 }
