[Interchange] Fixed crash importing obj/gltf/datasmith/mtlx when a texture is missing

diff --git a/Engine/Plugins/Experimental/Enterprise/DatasmithInterchange/Source/DatasmithInterchange/Private/InterchangeDatasmithTranslator.cpp b/Engine/Plugins/Experimental/Enterprise/DatasmithInterchange/Source/DatasmithInterchange/Private/InterchangeDatasmithTranslator.cpp
--- a/Engine/Plugins/Experimental/Enterprise/DatasmithInterchange/Source/DatasmithInterchange/Private/InterchangeDatasmithTranslator.cpp
+++ b/Engine/Plugins/Experimental/Enterprise/DatasmithInterchange/Source/DatasmithInterchange/Private/InterchangeDatasmithTranslator.cpp
@@ -545,6 +545,8 @@ TOptional<UE::Interchange::FImportImage> UInterchangeDatasmithTranslator::GetTex
 		return TOptional<UE::Interchange::FImportImage>();
 	}
 
+	SourceTranslator->SetResultsContainer(Results);
+
 	AlternateTexturePath = TextureElement->GetFile();
 
 	return TextureTranslator->GetTexturePayloadData(PayloadKey, AlternateTexturePath);
diff --git a/Engine/Plugins/Interchange/Runtime/Source/Import/Private/Gltf/InterchangeGltfTranslator.cpp b/Engine/Plugins/Interchange/Runtime/Source/Import/Private/Gltf/InterchangeGltfTranslator.cpp
--- a/Engine/Plugins/Interchange/Runtime/Source/Import/Private/Gltf/InterchangeGltfTranslator.cpp
+++ b/Engine/Plugins/Interchange/Runtime/Source/Import/Private/Gltf/InterchangeGltfTranslator.cpp
@@ -1412,6 +1412,7 @@ TOptional< UE::Interchange::FImportImage > UInterchangeGltfTranslator::GetTextur
 		{
 			return TOptional<UE::Interchange::FImportImage>();
 		}
+		SourceTranslator->SetResultsContainer(Results);
 
 		AlternateTexturePath = TextureFilePath;
 
diff --git a/Engine/Plugins/Interchange/Runtime/Source/Import/Private/MaterialX/InterchangeMaterialXTranslator.cpp b/Engine/Plugins/Interchange/Runtime/Source/Import/Private/MaterialX/InterchangeMaterialXTranslator.cpp
--- a/Engine/Plugins/Interchange/Runtime/Source/Import/Private/MaterialX/InterchangeMaterialXTranslator.cpp
+++ b/Engine/Plugins/Interchange/Runtime/Source/Import/Private/MaterialX/InterchangeMaterialXTranslator.cpp
@@ -230,6 +230,8 @@ TOptional<UE::Interchange::FImportImage> UInterchangeMaterialXTranslator::GetTex
 		return TOptional<UE::Interchange::FImportImage>();
 	}
 
+	SourceTranslator->SetResultsContainer(Results);
+
 	AlternateTexturePath = Filename;
 
 	TOptional<UE::Interchange::FImportImage> TexturePayloadData = TextureTranslator->GetTexturePayloadData(PayloadKey, AlternateTexturePath);
diff --git a/Engine/Plugins/Interchange/Runtime/Source/Import/Private/Mesh/InterchangeOBJTranslator.cpp b/Engine/Plugins/Interchange/Runtime/Source/Import/Private/Mesh/InterchangeOBJTranslator.cpp
--- a/Engine/Plugins/Interchange/Runtime/Source/Import/Private/Mesh/InterchangeOBJTranslator.cpp
+++ b/Engine/Plugins/Interchange/Runtime/Source/Import/Private/Mesh/InterchangeOBJTranslator.cpp
@@ -1093,6 +1093,8 @@ TOptional<UE::Interchange::FImportImage> UInterchangeOBJTranslator::GetTexturePa
 		return TOptional<UE::Interchange::FImportImage>();
 	}
 
+	SourceTranslator->SetResultsContainer(Results);
+
 	AlternateTexturePath = PayLoadKey;
 
 	return TextureTranslator->GetTexturePayloadData(PayLoadKey, AlternateTexturePath);
