Interchange Area Light Import Fix
- Area Lights are imported correctly now.
- IES based lights are deprecated and would fallback to Point lights.
- Area lights will use intelligent defaults now instead of -1.0f value for distances.

diff --git a/Engine/Plugins/Enterprise/DatasmithContent/Source/DatasmithContent/Public/DatasmithAreaLightActor.h b/Engine/Plugins/Enterprise/DatasmithContent/Source/DatasmithContent/Public/DatasmithAreaLightActor.h
--- a/Engine/Plugins/Enterprise/DatasmithContent/Source/DatasmithContent/Public/DatasmithAreaLightActor.h
+++ b/Engine/Plugins/Enterprise/DatasmithContent/Source/DatasmithContent/Public/DatasmithAreaLightActor.h
@@ -23,6 +23,8 @@ enum class EDatasmithAreaLightActorType : uint8
 {
 	Point,
 	Spot,
+	IES_DEPRECATED, // Added for consistency with type EDatasmithAreaLightType. 
+					// Removing this will result in 'EDatasmithAreaLightActorType MAX' enum value in the editor.
 	Rect,
 };
 
diff --git a/Engine/Plugins/Experimental/Enterprise/DatasmithInterchange/Source/DatasmithInterchange/Private/InterchangeDatasmithAreaLightFactory.cpp b/Engine/Plugins/Experimental/Enterprise/DatasmithInterchange/Source/DatasmithInterchange/Private/InterchangeDatasmithAreaLightFactory.cpp
--- a/Engine/Plugins/Experimental/Enterprise/DatasmithInterchange/Source/DatasmithInterchange/Private/InterchangeDatasmithAreaLightFactory.cpp
+++ b/Engine/Plugins/Experimental/Enterprise/DatasmithInterchange/Source/DatasmithInterchange/Private/InterchangeDatasmithAreaLightFactory.cpp
@@ -82,5 +82,5 @@ UObject* UInterchangeDatasmithAreaLightFactory::ProcessActor(AActor& SpawnedActo
 #endif //WITH_EDITOR
 	}
 
-	return AreaLightActor;
+	return Super::ProcessActor(SpawnedActor, FactoryNode, NodeContainer);
 };
\ No newline at end of file
diff --git a/Engine/Plugins/Experimental/Enterprise/DatasmithInterchange/Source/DatasmithInterchange/Private/InterchangeDatasmithLevelPipeline.cpp b/Engine/Plugins/Experimental/Enterprise/DatasmithInterchange/Source/DatasmithInterchange/Private/InterchangeDatasmithLevelPipeline.cpp
--- a/Engine/Plugins/Experimental/Enterprise/DatasmithInterchange/Source/DatasmithInterchange/Private/InterchangeDatasmithLevelPipeline.cpp
+++ b/Engine/Plugins/Experimental/Enterprise/DatasmithInterchange/Source/DatasmithInterchange/Private/InterchangeDatasmithLevelPipeline.cpp
@@ -177,9 +177,67 @@ void UInterchangeDatasmithLevelPipeline::SetUpFactoryNode(UInterchangeActorFacto
 		FactoryNode->SetCustom##AttributeName(AttributeName);\
 	}\
 
-void UInterchangeDatasmithLevelPipeline::SetupAreaLight(UInterchangeDatasmithAreaLightFactoryNode* FactoryNode, const UInterchangeDatasmithAreaLightNode* TranslatedNode) const
+/**
+* Only modifies the attribute if a condition is satisfied
+* Condition should an expression evaluated to a boolean.
+*/
+#define APPLY_FACTORY_ATTRIBUTE_WITH_VALIDATION(AttributeName, AttributeType, Condition) \
+	AttributeType AttributeName;\
+	if (TranslatedNode->GetCustom##AttributeName(AttributeName))\
+	{\
+		if(Condition)\
+		{\
+			FactoryNode->SetCustom##AttributeName(AttributeName);\
+		}\
+	}\
+
+/**
+ * Allows Attributes to be modified before assignment.
+ * A modifier could be a function, or a type conversion or anything else that can take the attribute value
+ * as a parameter.
+ */
+#define APPLY_FACTORY_ATTRIBUTE_WITH_MODIFIER(AttributeName, AttributeType, Modifier) \
+	AttributeType AttributeName;\
+	if (TranslatedNode->GetCustom##AttributeName(AttributeName))\
+	{\
+		FactoryNode->SetCustom##AttributeName(Modifier(AttributeName));\
+	}\
+
+namespace UE::Interchange::AreaLightUtils
 {
-	APPLY_FACTORY_ATTRIBUTE(LightType, EDatasmithAreaLightActorType);
+	// Helper function to correctly assign types that have been deprecated.
+	EDatasmithAreaLightActorType GetLightActorType(const EDatasmithAreaLightActorType LightType)
+	{
+		EDatasmithAreaLightActorType LightActorType = EDatasmithAreaLightActorType::Point;
+
+		switch (LightType)
+		{
+		case EDatasmithAreaLightActorType::Spot:
+			LightActorType = EDatasmithAreaLightActorType::Spot;
+			break;
+
+		case EDatasmithAreaLightActorType::Point:
+			LightActorType = EDatasmithAreaLightActorType::Point;
+			break;
+
+		case EDatasmithAreaLightActorType::IES_DEPRECATED:
+			LightActorType = EDatasmithAreaLightActorType::Point;
+			break;
+
+		case EDatasmithAreaLightActorType::Rect:
+			LightActorType = EDatasmithAreaLightActorType::Rect;
+			break;
+		}
+
+		return LightActorType;
+	}
+}
+
+void UInterchangeDatasmithLevelPipeline::SetupAreaLight(UInterchangeDatasmithAreaLightFactoryNode* FactoryNode, const UInterchangeDatasmithAreaLightNode* TranslatedNode) const
+{	
+	using namespace UE::Interchange::AreaLightUtils;
+
+	APPLY_FACTORY_ATTRIBUTE_WITH_MODIFIER(LightType, EDatasmithAreaLightActorType, GetLightActorType);
 	APPLY_FACTORY_ATTRIBUTE(LightShape, EDatasmithAreaLightActorShape);
 	APPLY_FACTORY_ATTRIBUTE(Dimensions, FVector2D);
 	
@@ -197,11 +255,14 @@ void UInterchangeDatasmithLevelPipeline::SetupAreaLight(UInterchangeDatasmithAre
 	APPLY_FACTORY_ATTRIBUTE(IESBrightnessScale, float);
 	APPLY_FACTORY_ATTRIBUTE(Rotation, FRotator);
 
-	APPLY_FACTORY_ATTRIBUTE(SourceRadius, float);
-	APPLY_FACTORY_ATTRIBUTE(SourceLength, float);
-	APPLY_FACTORY_ATTRIBUTE(AttenuationRadius, float);
+	APPLY_FACTORY_ATTRIBUTE_WITH_VALIDATION(SourceRadius, float, SourceRadius > 0.0f);
+	APPLY_FACTORY_ATTRIBUTE_WITH_VALIDATION(SourceLength, float, SourceLength > 0.0f);
+	APPLY_FACTORY_ATTRIBUTE_WITH_VALIDATION(AttenuationRadius, float, AttenuationRadius > 0.0f);
+
 	APPLY_FACTORY_ATTRIBUTE(SpotlightInnerAngle, float);
 	APPLY_FACTORY_ATTRIBUTE(SpotlightOuterAngle, float);
 }
 
-#undef APPLY_FACTORY_ATTRIBUTE
\ No newline at end of file
+#undef APPLY_FACTORY_ATTRIBUTE
+#undef APPLY_FACTORY_ATTRIBUTE_WITH_VALIDATION
+#undef APPLY_FACTORY_ATTRIBUTE_WITH_MODIFIER
\ No newline at end of file
