DatasmithIntechange - Difference in output logs when importing udatasmith file via Interchange import vs File Import

- Now the static meshes are analysed in the PostImport step of the pipline and are discarded if found to be invalid.
- Also has side effects on the StaticMeshActors, as just deleting the StaticMesh assets is not enough, the references need to be cleared as well.

diff --git a/Engine/Plugins/Experimental/Enterprise/DatasmithInterchange/Source/DatasmithInterchange/Private/InterchangeDatasmithPipeline.cpp b/Engine/Plugins/Experimental/Enterprise/DatasmithInterchange/Source/DatasmithInterchange/Private/InterchangeDatasmithPipeline.cpp
--- a/Engine/Plugins/Experimental/Enterprise/DatasmithInterchange/Source/DatasmithInterchange/Private/InterchangeDatasmithPipeline.cpp
+++ b/Engine/Plugins/Experimental/Enterprise/DatasmithInterchange/Source/DatasmithInterchange/Private/InterchangeDatasmithPipeline.cpp
@@ -21,12 +21,17 @@
 #include "InterchangeTexture2DNode.h"
 #include "InterchangeTextureFactoryNode.h"
 #include "InterchangeTexture2DFactoryNode.h"
+#include "InterchangeMeshActorFactoryNode.h"
 
 #include "ExternalSource.h"
 #include "DatasmithAreaLightActor.h"
 #include "DatasmithScene.h"
 #include "DatasmithSceneXmlWriter.h"
 
+#include "StaticMeshAttributes.h"
+#include "Engine/StaticMeshActor.h"
+#include "Components/StaticMeshComponent.h"
+
 #if WITH_EDITOR
 #include "DatasmithImporter.h"
 #include "DatasmithImportContext.h"
@@ -37,6 +42,63 @@
 
 #define LOCTEXT_NAMESPACE "InterchangeDatasmithPipeline"
 
+namespace UE::Interchange::StaticMeshUtils
+{
+#if WITH_EDITORONLY_DATA
+	// Mirror of DatasmithMeshHelper::IsMeshValid
+	bool HasValidTriangleData(const FMeshDescription& MeshDescription, const FMeshBuildSettings& BuildSettings)
+	{
+		const FVector BuildScale = BuildSettings.BuildScale3D;
+		TVertexAttributesConstRef<FVector3f> VertexPositions = MeshDescription.GetVertexPositions();
+		FVector3f RawNormalScale(BuildScale.Y * BuildScale.Z, BuildScale.X * BuildScale.Z, BuildScale.X * BuildScale.Y); // Component-wise scale
+
+		for (const FTriangleID TriangleID : MeshDescription.Triangles().GetElementIDs())
+		{
+			TArrayView<const FVertexID> VertexIDs = MeshDescription.GetTriangleVertices(TriangleID);
+			FVector3f Corners[3] =
+			{
+				VertexPositions[VertexIDs[0]],
+				VertexPositions[VertexIDs[1]],
+				VertexPositions[VertexIDs[2]]
+			};
+
+			FVector3f RawNormal = (Corners[1] - Corners[2]) ^ (Corners[0] - Corners[2]);
+			RawNormal *= RawNormalScale;
+			float FourSquaredTriangleArea = RawNormal.SizeSquared();
+
+			// We support even small triangles, but this function is still useful to
+			// see if we have at least one valid triangle in the mesh
+			if (FourSquaredTriangleArea > 0.0f)
+			{
+				return true;
+			}
+		}
+
+		// all faces are degenerated, mesh is invalid
+		return false;
+	}
+
+	bool IsMeshValid(const UStaticMesh* StaticMesh)
+	{
+		// LOD index 0 has to be valid for a mesh to be valid.
+		if (!StaticMesh->IsMeshDescriptionValid(0))
+		{
+			return false;
+		}
+
+		const FMeshBuildSettings& BuildSettings = StaticMesh->GetSourceModel(0).BuildSettings;
+		const FMeshDescription& MeshDescription = *StaticMesh->GetMeshDescription(0);
+
+		if (HasValidTriangleData(MeshDescription, BuildSettings))
+		{
+			return true;
+		}
+
+		return false;
+	}
+#endif
+}
+
 UInterchangeDatasmithPipeline::UInterchangeDatasmithPipeline()
 {
 	CommonSkeletalMeshesAndAnimationsProperties = CreateDefaultSubobject<UInterchangeGenericCommonSkeletalMeshesAndAnimationsProperties>("CommonSkeletalMeshesAndAnimationsProperties");
@@ -52,6 +114,7 @@ UInterchangeDatasmithPipeline::UInterchangeDatasmithPipeline()
 	AnimationPipeline = CreateDefaultSubobject<UInterchangeGenericAnimationPipeline>("AnimationPipeline");
 
 	MeshPipeline->CommonMeshesProperties = CommonMeshesProperties;
+	//MeshPipeline->bCombineStaticMeshes = false;
 	MeshPipeline->CommonSkeletalMeshesAndAnimationsProperties = CommonSkeletalMeshesAndAnimationsProperties;
 
 	AnimationPipeline->CommonMeshesProperties = CommonMeshesProperties;
@@ -103,7 +166,25 @@ void UInterchangeDatasmithPipeline::ExecutePipeline(UInterchangeBaseNodeContaine
 	for (UInterchangeBaseMaterialFactoryNode* MaterialFactoryNode : NodeUtils::GetNodes<UInterchangeBaseMaterialFactoryNode>(BaseNodeContainer))
 	{
 		DependenciesUids.Add(MaterialFactoryNode->GetUniqueID());
-		MaterialFactoryNode->SetCustomSubPath(FPaths::Combine(PackageSubPath, "Materials"));
+		if (bCreateMaterialReferencesFolders)
+		{
+			if (MaterialFactoryNode->IsA<UInterchangeMaterialFactoryNode>())
+			{
+				MaterialFactoryNode->SetCustomSubPath(FPaths::Combine(PackageSubPath, "Materials/References"));
+			}
+			else if(MaterialFactoryNode->IsA<UInterchangeMaterialFunctionFactoryNode>())
+			{
+				MaterialFactoryNode->SetCustomSubPath(FPaths::Combine(PackageSubPath, "Materials/References/Functions"));
+			}
+			else
+			{
+				MaterialFactoryNode->SetCustomSubPath(FPaths::Combine(PackageSubPath, "Materials"));
+			}
+		}
+		else
+		{
+			MaterialFactoryNode->SetCustomSubPath(FPaths::Combine(PackageSubPath, "Materials"));
+		}
 		MaterialFactoryNode->SetEnabled(true);
 	}
 
@@ -188,6 +269,56 @@ void UInterchangeDatasmithPipeline::ExecutePostImportPipeline(const UInterchange
 	{
 		PostImportDatasmithSceneAsset(*DatasmithSceneAsset);
 	}
+#if WITH_EDITORONLY_DATA
+	if (bDeleteInvalidMeshes)
+	{
+		if (UStaticMesh* StaticMesh = Cast<UStaticMesh>(CreatedAsset))
+		{
+			if (!UE::Interchange::StaticMeshUtils::IsMeshValid(StaticMesh))
+			{
+				UInterchangeResultWarning_Generic* InvalidMeshWarning = AddMessage<UInterchangeResultWarning_Generic>();
+				InvalidMeshWarning->Text = FText::Format(
+					LOCTEXT("InvalidStaticMesh", "Static Mesh {0} contains only degenerate or empty triangles."),
+					FText::FromString(StaticMesh->GetName())
+				);
+
+				if (UInterchangeStaticMeshFactoryNode* FactoryNode = Cast<UInterchangeStaticMeshFactoryNode>(InBaseNodeContainer->GetFactoryNode(NodeKey)))
+				{
+					FactoryNode->SetCustomReferenceObject(FSoftObjectPath());
+					InvalidStaticMeshFactoryUids.Emplace(FactoryNode->GetUniqueID());
+				}
+
+				CreatedAsset->Rename(nullptr, GetTransientPackage(), REN_NonTransactional | REN_DontCreateRedirectors);
+				if (CreatedAsset->IsRooted())
+				{
+					CreatedAsset->RemoveFromRoot();
+				}
+				CreatedAsset->ClearFlags(RF_Public | RF_Standalone);
+				CreatedAsset->MarkAsGarbage();
+			}
+		}
+
+		// Clear up the reference 
+		if (AStaticMeshActor* StaticMeshActor = Cast<AStaticMeshActor>(CreatedAsset))
+		{
+			if (UInterchangeMeshActorFactoryNode* FactoryNode = Cast<UInterchangeMeshActorFactoryNode>(InBaseNodeContainer->GetFactoryNode(NodeKey)))
+			{
+				TArray<FString> FactoryDependencies;
+				FactoryNode->GetFactoryDependencies(FactoryDependencies);
+				for (const FString& FactoryDependency : FactoryDependencies)
+				{
+					if (InvalidStaticMeshFactoryUids.Contains(FactoryDependency))
+					{
+						if (UStaticMeshComponent* StaticMeshComponent = StaticMeshActor->GetStaticMeshComponent())
+						{
+							StaticMeshComponent->SetStaticMesh(nullptr);
+						}
+					}
+				}
+			}
+		}
+	}
+#endif
 }
 
 void UInterchangeDatasmithPipeline::PostImportDatasmithSceneAsset(UDatasmithScene& DatasmithSceneAsset)
diff --git a/Engine/Plugins/Experimental/Enterprise/DatasmithInterchange/Source/DatasmithInterchange/Public/InterchangeDatasmithPipeline.h b/Engine/Plugins/Experimental/Enterprise/DatasmithInterchange/Source/DatasmithInterchange/Public/InterchangeDatasmithPipeline.h
--- a/Engine/Plugins/Experimental/Enterprise/DatasmithInterchange/Source/DatasmithInterchange/Public/InterchangeDatasmithPipeline.h
+++ b/Engine/Plugins/Experimental/Enterprise/DatasmithInterchange/Source/DatasmithInterchange/Public/InterchangeDatasmithPipeline.h
@@ -50,6 +50,14 @@ class DATASMITHINTERCHANGE_API UInterchangeDatasmithPipeline : public UInterchan
 	UPROPERTY(VisibleAnywhere, Instanced, Category = "Common Skeletal Meshes and Animations")
 	TObjectPtr<UInterchangeGenericCommonSkeletalMeshesAndAnimationsProperties> CommonSkeletalMeshesAndAnimationsProperties;
 
+#if WITH_EDITORONLY_DATA
+	UPROPERTY(EditAnywhere, Category = "Mesh")
+	bool bDeleteInvalidMeshes = true;
+#endif
+
+	UPROPERTY(EditAnywhere, Category = "Materials")
+	bool bCreateMaterialReferencesFolders = true;
+
 protected:
 	virtual void ExecutePipeline(UInterchangeBaseNodeContainer* InBaseNodeContainer, const TArray<UInterchangeSourceData*>& SourceDatas) override;
 	virtual void ExecutePostImportPipeline(const UInterchangeBaseNodeContainer* InBaseNodeContainer, const FString& NodeKey, UObject* CreatedAsset, bool bIsAReimport) override;
@@ -65,4 +73,6 @@ private:
 
 	// Fill up the UDatasmithScene with all the data its needs for DnD
 	void PostImportDatasmithSceneAsset(UDatasmithScene& DatasmithSceneAsset);
+
+	TSet<FString> InvalidStaticMeshFactoryUids;
 };
