Fixed crash when loading files with multiple instances of the same mesh

#rnx

diff --git a/Engine/Plugins/Experimental/Enterprise/DatasmithInterchange/Source/DatasmithInterchange/Private/InterchangeDatasmithTranslator.cpp b/Engine/Plugins/Experimental/Enterprise/DatasmithInterchange/Source/DatasmithInterchange/Private/InterchangeDatasmithTranslator.cpp
--- a/Engine/Plugins/Experimental/Enterprise/DatasmithInterchange/Source/DatasmithInterchange/Private/InterchangeDatasmithTranslator.cpp
+++ b/Engine/Plugins/Experimental/Enterprise/DatasmithInterchange/Source/DatasmithInterchange/Private/InterchangeDatasmithTranslator.cpp
@@ -578,7 +578,7 @@ TFuture<TOptional<UE::Interchange::FMeshPayloadData>> UInterchangeDatasmithTrans
 		return EmptyPromise.GetFuture();
 	}
 
-	return Async(EAsyncExecution::TaskGraph, [this, MeshElement = MoveTemp(MeshElement), &MeshGlobalTransform]
+	return Async(EAsyncExecution::TaskGraph, [this, MeshElement = MoveTemp(MeshElement), MeshGlobalTransform]
 		{
 			TOptional<UE::Interchange::FMeshPayloadData> Result;
 
