FIX: non unity build - Reference Materials (Issue 347829)
- Added missing includes

#rnx

diff --git a/Engine/Plugins/Experimental/Enterprise/DatasmithInterchange/Source/DatasmithInterchange/Private/InterchangeDatasmithMaterialPipeline.cpp b/Engine/Plugins/Experimental/Enterprise/DatasmithInterchange/Source/DatasmithInterchange/Private/InterchangeDatasmithMaterialPipeline.cpp
--- a/Engine/Plugins/Experimental/Enterprise/DatasmithInterchange/Source/DatasmithInterchange/Private/InterchangeDatasmithMaterialPipeline.cpp
+++ b/Engine/Plugins/Experimental/Enterprise/DatasmithInterchange/Source/DatasmithInterchange/Private/InterchangeDatasmithMaterialPipeline.cpp
@@ -26,6 +26,7 @@
 #include "Materials/MaterialInstanceConstant.h"
 #include "Materials/MaterialInstanceDynamic.h"
 #include "Misc/PackageName.h"
+#include "UObject/ObjectRedirector.h"
 
 DEFINE_LOG_CATEGORY(LogInterchangeMaterialPipeline);
 
diff --git a/Engine/Plugins/Interchange/Runtime/Source/Import/Private/Material/InterchangeMaterialFactory.cpp b/Engine/Plugins/Interchange/Runtime/Source/Import/Private/Material/InterchangeMaterialFactory.cpp
--- a/Engine/Plugins/Interchange/Runtime/Source/Import/Private/Material/InterchangeMaterialFactory.cpp
+++ b/Engine/Plugins/Interchange/Runtime/Source/Import/Private/Material/InterchangeMaterialFactory.cpp
@@ -29,6 +29,7 @@
 #include "Nodes/InterchangeBaseNodeContainer.h"
 #include "Nodes/InterchangeUserDefinedAttribute.h"
 #include "UObject/ObjectRedirector.h"
+#include "Misc/PackageName.h"
 
 #include UE_INLINE_GENERATED_CPP_BY_NAME(InterchangeMaterialFactory)
 
