Add the ability to import a Datasmith file from the 'Import into level' menu entry

#rnx
#okforgihub public

#virtualized

diff --git a/Engine/Plugins/Experimental/Enterprise/DatasmithInterchange/DatasmithInterchange.uplugin b/Engine/Plugins/Experimental/Enterprise/DatasmithInterchange/DatasmithInterchange.uplugin
--- a/Engine/Plugins/Experimental/Enterprise/DatasmithInterchange/DatasmithInterchange.uplugin
+++ b/Engine/Plugins/Experimental/Enterprise/DatasmithInterchange/DatasmithInterchange.uplugin
@@ -11,7 +11,7 @@
 	"MarketplaceURL" : "",
 	"SupportURL" : "",
 	"EnabledByDefault" : false,
-	"CanContainContent" : false,
+	"CanContainContent" : true,
 	"IsBetaVersion" : true,
 	"Installed" : false,
 	"PlatformAllowList": [
@@ -24,11 +24,6 @@
 			"Name": "DatasmithInterchange",
 			"Type": "Runtime",
 			"LoadingPhase": "Default"
-		},
-		{
-			"Name": "DatasmithInterchangeScripting",
-			"Type": "Runtime",
-			"LoadingPhase": "Default"
 		}
 	],
 	"Plugins": [
@@ -38,7 +33,22 @@
 		},
 		{
 			"Name": "DatasmithContent",
-			"Enabled":  true
+			"Enabled": true
+		},
+		{
+			"Name": "DatasmithCADImporter",
+			"Enabled": true,
+			"SupportedTargetPlatforms": [
+				"Win64"
+			]
+		},
+		{
+			"Name": "DatasmithC4DImporter",
+			"Enabled": true,
+			"SupportedTargetPlatforms": [
+				"Win64",
+				"Mac"
+			]
 		},
 		{
 			"Name": "Interchange",
diff --git a/Engine/Plugins/Experimental/Enterprise/DatasmithInterchange/Source/DatasmithInterchange/Private/DatasmithInterchangeModule.cpp b/Engine/Plugins/Experimental/Enterprise/DatasmithInterchange/Source/DatasmithInterchange/Private/DatasmithInterchangeModule.cpp
--- a/Engine/Plugins/Experimental/Enterprise/DatasmithInterchange/Source/DatasmithInterchange/Private/DatasmithInterchangeModule.cpp
+++ b/Engine/Plugins/Experimental/Enterprise/DatasmithInterchange/Source/DatasmithInterchange/Private/DatasmithInterchangeModule.cpp
@@ -6,7 +6,6 @@
 #include "InterchangeDatasmithAreaLightFactory.h"
 #include "InterchangeDatasmithLog.h"
 #include "InterchangeDatasmithPipeline.h"
-#include "InterchangeDatasmithSceneFactory.h"
 #include "InterchangeDatasmithTranslator.h"
 
 #include "InterchangeReferenceMaterials/DatasmithReferenceMaterialManager.h"
@@ -21,22 +20,9 @@
 #include "InterchangeGenericMeshPipeline.h"
 #include "InterchangeGenericScenesPipeline.h"
 #include "InterchangeManager.h"
+#include "InterchangeProjectSettings.h"
 #include "Logging/LogMacros.h"
-
-#if WITH_EDITOR
-#include "DatasmithImporterModule.h"
-#include "DatasmithUtils.h"
-#include "DesktopPlatformModule.h"
-#include "Dialogs/DlgPickPath.h"
-#include "EditorDirectories.h"
-#include "IDesktopPlatform.h"
-#include "Interfaces/IMainFrameModule.h"
-#include "ObjectTools.h"
-#include "Styling/SlateIconFinder.h"
-#include "UI/DatasmithImportOptionsWindow.h"
-#include "ToolMenus.h"
-#include "ToolMenuSection.h"
-#endif //WITH_EDITOR
+#include "UObject/SoftObjectPath.h"
 
 #define LOCTEXT_NAMESPACE "DatasmithInterchange"
 
@@ -56,13 +42,23 @@ public:
 
 		UInterchangeManager& InterchangeManager = UInterchangeManager::GetInterchangeManager();
 		InterchangeManager.RegisterTranslator(UInterchangeDatasmithTranslator::StaticClass());
-		InterchangeManager.RegisterFactory(UInterchangeDatasmithSceneFactory::StaticClass());
 		InterchangeManager.RegisterFactory(UInterchangeDatasmithAreaLightFactory::StaticClass());
 
-#if WITH_EDITOR
-		IDatasmithImporterModule& DatasmithImporterModule = IDatasmithImporterModule::Get();
-		DatasmithImporterModule.OnGenerateDatasmithImportMenu().AddRaw(this, &FDatasmithInterchangeModule::ExtendDatasmitMenuOptions);
-#endif //WITH_EDITOR
+		// Add Datasmith translator to the Interchange engine
+		FInterchangeTranslatorPipelines TranslatorPipeplines;
+		TranslatorPipeplines.Translator = TSoftClassPtr<UInterchangeTranslatorBase>(UInterchangeDatasmithTranslator::StaticClass());
+		TranslatorPipeplines.Pipelines.Add(FSoftObjectPath(TEXT("/DatasmithInterchange/InterchangeDatasmithDefault.InterchangeDatasmithDefault")));
+		{
+			FInterchangeImportSettings& ImportSettings = FInterchangeProjectSettingsUtils::GetMutableDefaultImportSettings(false);
+			FInterchangePipelineStack& PipelineStack = ImportSettings.PipelineStacks[TEXT("Assets")];
+			PipelineStack.PerTranslatorPipelines.Add(TranslatorPipeplines);
+		}
+
+		{
+			FInterchangeImportSettings& ImportSettings = FInterchangeProjectSettingsUtils::GetMutableDefaultImportSettings(true);
+			FInterchangePipelineStack& PipelineStack = ImportSettings.PipelineStacks[TEXT("Scene")];
+			PipelineStack.PerTranslatorPipelines.Add(TranslatorPipeplines);
+		}
 
 		FDatasmithReferenceMaterialManager::Create();
 
@@ -76,207 +72,10 @@ public:
 
 	virtual void ShutdownModule() override
 	{
-#if WITH_EDITOR
-		if (IDatasmithImporterModule::IsAvailable())
-		{
-			IDatasmithImporterModule& DatasmithImporterModule = IDatasmithImporterModule::Get();
-			DatasmithImporterModule.OnGenerateDatasmithImportMenu().RemoveAll(this);
-		}
-#endif //WITH_EDITOR
-
 		UE::DatasmithInterchange::FDatasmithReferenceMaterialManager::Destroy();
 	}
-
-private:
-#if WITH_EDITOR
-	void ExtendDatasmitMenuOptions(FToolMenuSection& SubSection);
-
-	void OnImportInterchange();
-
-	bool DisplayOptionsDialog(const FString& FilePath, const FString& PackagePath, const TArray<TObjectPtr<UDatasmithOptionsBase>>& ImportOptions);
-		
-	bool Import(const FString& FilePath, const FString& ContentPath);
-#endif //WITH_EDITOR
 };
 
-#if WITH_EDITOR
-void FDatasmithInterchangeModule::ExtendDatasmitMenuOptions(FToolMenuSection& SubSection)
-{
-	SubSection.AddMenuEntry(
-		TEXT("InterchangeImportFile"),
-		LOCTEXT("DatasmithInterchangeFileImport", "Interchange Import..."), // label
-		LOCTEXT("DatasmithInterchangeFileImportTooltip", "Experimental: Import Unreal Datasmith file using Interchange"), // description
-		FSlateIcon(FSlateIconFinder::FindIcon("Datasmith.Import")),
-		//FSlateIcon(FDatasmithStyle::GetStyleSetName(), TEXT("Datasmith.Import")),
-		FUIAction(FExecuteAction::CreateRaw(this, &FDatasmithInterchangeModule::OnImportInterchange))
-	);
-}
-
-void FDatasmithInterchangeModule::OnImportInterchange()
-{
-	const TArray<FString> OldFormats = FDatasmithTranslatorManager::Get().GetSupportedFormats();
-	TArray<FString> Formats;
-	Formats.Reserve(OldFormats.Num());
-
-	for (const FString& Format : OldFormats)
-	{
-		if (Format.Contains(TEXT("gltf")) || Format.Contains(TEXT("glb")))
-		{
-			continue;
-		}
-
-		Formats.Add(Format);
-	}
-
-	FString FileTypes;
-	FString Extensions;
-
-	FString DefaultLocation(FEditorDirectories::Get().GetLastDirectory(ELastDirectory::GENERIC_IMPORT));
-	ObjectTools::AppendFormatsFileExtensions(Formats, FileTypes, Extensions);
-
-	const FString FormatString = FString::Printf(TEXT("All Files (%s)|%s|%s"), *Extensions, *Extensions, *FileTypes);
-	const FString Title = LOCTEXT("BrowseSourceDialogTitle", "Select Datasmith Source File").ToString();
-
-	TArray<FString> OutOpenFilenames;
-	IDesktopPlatform* DesktopPlatform = FDesktopPlatformModule::Get();
-	bool bOpened = false;
-	if (DesktopPlatform)
-	{
-		bOpened = DesktopPlatform->OpenFileDialog(
-			FSlateApplication::Get().FindBestParentWindowHandleForDialogs(nullptr),
-			Title,
-			DefaultLocation,
-			FString(),
-			FormatString,
-			EFileDialogFlags::None,
-			OutOpenFilenames
-		);
-	}
-
-	if (bOpened && OutOpenFilenames.Num() > 0)
-	{
-		static FString ContentPath = "/Game/"; // Trailing '/' is needed to set the default path
-
-		TSharedRef<SDlgPickPath> PickContentPathDlg =
-			SNew(SDlgPickPath)
-			.Title(LOCTEXT("ChooseImportRootContentPath", "Choose Location for importing the Datasmith content"))
-			.DefaultPath(FText::FromString(ContentPath));
-
-		if (PickContentPathDlg->ShowModal() != EAppReturnType::Cancel)
-		{
-			FEditorDirectories::Get().SetLastDirectory(ELastDirectory::GENERIC_IMPORT, FPaths::GetPath(OutOpenFilenames[0])); // Save path as default for next time.
-
-			ContentPath = PickContentPathDlg->GetPath().ToString() + "/";
-			Import(OutOpenFilenames[0], ContentPath);
-		}
-	}
-}
-
-bool FDatasmithInterchangeModule::DisplayOptionsDialog(const FString& FilePath, const FString& PackagePath, const TArray<TObjectPtr<UDatasmithOptionsBase>>& ImportOptionsPtr)
-{
-	if (ImportOptionsPtr.Num() == 0)
-	{
-		return true;
-	}
-
-	TSharedPtr<SWindow> ParentWindow;
-
-	if (FModuleManager::Get().IsModuleLoaded("MainFrame"))
-	{
-		IMainFrameModule& MainFrame = FModuleManager::LoadModuleChecked<IMainFrameModule>("MainFrame");
-		ParentWindow = MainFrame.GetParentWindow();
-	}
-
-	TArray<UObject*> ImportOptions;
-	ImportOptions.SetNum(ImportOptionsPtr.Num());
-	for (int32 Index = 0; Index < ImportOptionsPtr.Num(); ++Index)
-	{
-		ImportOptions[Index] = ImportOptionsPtr[Index];
-	}
-
-	TSharedRef<SWindow> Window = SNew(SWindow)
-		.Title(LOCTEXT("DatasmithImportSettingsTitle", "Datasmith Import Options"))
-		.SizingRule(ESizingRule::Autosized);
-
-	float SceneVersion = FDatasmithUtils::GetDatasmithFormatVersionAsFloat();
-	FString FileSDKVersion = FDatasmithUtils::GetEnterpriseVersionAsString();
-
-	TSharedPtr<SDatasmithOptionsWindow> OptionsWindow;
-	Window->SetContent
-	(
-		SAssignNew(OptionsWindow, SDatasmithOptionsWindow)
-		.ImportOptions(ImportOptions)
-		.WidgetWindow(Window)
-		// note: Spacing in text below is intentional for text alignment
-		.FileNameText(FText::Format(LOCTEXT("DatasmithImportSettingsFileName", "  Import File  :    {0}"), FText::FromString(FPaths::GetCleanFilename(FilePath))))
-		.FilePathText(FText::FromString(FilePath))
-		.FileFormatVersion(SceneVersion)
-		.FileSDKVersion(FText::FromString(FileSDKVersion))
-		.PackagePathText(FText::Format(LOCTEXT("DatasmithImportSettingsPackagePath", "  Import To   :    {0}"), FText::FromString(PackagePath)))
-		.ProceedButtonLabel(LOCTEXT("DatasmithOptionWindow_ImportCurLevel", "Import"))
-		.ProceedButtonTooltip(LOCTEXT("DatasmithOptionWindow_ImportCurLevel_ToolTip", "Import the file through Interchange and add to the current Level"))
-		.CancelButtonLabel(LOCTEXT("DatasmithOptionWindow_Cancel", "Cancel"))
-		.CancelButtonTooltip(LOCTEXT("DatasmithOptionWindow_Cancel_ToolTip", "Cancel importing this file"))
-		.MinDetailHeight(320.f)
-		.MinDetailWidth(450.f)
-	);
-
-	FSlateApplication::Get().AddModalWindow(Window, ParentWindow, false);
-
-	return OptionsWindow->ShouldImport();
-}
-
-bool FDatasmithInterchangeModule::Import(const FString& FilePath, const FString& ContentPath)
-{
-	UE::Interchange::FScopedSourceData ScopedSourceData(FilePath);
-
-	UInterchangeTranslatorBase* InterchangeTranslator = UInterchangeManager::GetInterchangeManager().GetTranslatorForSourceData(ScopedSourceData.GetSourceData());
-	if (!InterchangeTranslator || !InterchangeTranslator->IsA(UInterchangeDatasmithTranslator::StaticClass()))
-	{
-		const FString FileExtension = FPaths::GetExtension(FilePath);
-		UE_LOG(LogInterchangeDatasmith, Warning, TEXT("File extension \"%s\" is not supported by DatasmithInterchange"), *FileExtension);
-
-		return false;
-	}
-
-	{
-		FDatasmithSceneSource DatasmithSource;
-		DatasmithSource.SetSourceFile(FilePath);
-
-		TSharedPtr<IDatasmithTranslator> Translator = FDatasmithTranslatorManager::Get().SelectFirstCompatible(DatasmithSource);
-		if (Translator)
-		{
-			TArray<TObjectPtr<UDatasmithOptionsBase>> ImportOptions;
-			Translator->GetSceneImportOptions(ImportOptions);
-
-			bool bShouldImport = DisplayOptionsDialog(FilePath, ContentPath, ImportOptions);
-			if (!bShouldImport)
-			{
-				return false;
-			}
-
-			const FString OptionFilePath = UInterchangeDatasmithTranslator::BuildConfigFilePath(FilePath);
-			for (TObjectPtr<UDatasmithOptionsBase>& Option : ImportOptions)
-			{
-				Option->SaveConfig(CPF_Config, *OptionFilePath);
-			}
-		}
-	}
-
-
-	FImportAssetParameters ImportAssetParameters;
-	ImportAssetParameters.bIsAutomated = false;
-
-	UInterchangeDatasmithPipeline* DatasmithPipeline = NewObject< UInterchangeDatasmithPipeline >();
-
-	ImportAssetParameters.OverridePipelines.Add(DatasmithPipeline);
-
-	TTuple<UE::Interchange::FAssetImportResultRef, UE::Interchange::FSceneImportResultRef> ImportResults
-		= UInterchangeManager::GetInterchangeManager().ImportSceneAsync(ContentPath, ScopedSourceData.GetSourceData(), ImportAssetParameters);
-	return ImportResults.Key->IsValid() && ImportResults.Value->IsValid();
-}
-#endif //WITH_EDITOR
-
 IMPLEMENT_MODULE(FDatasmithInterchangeModule, DatasmithInterchange);
 
 #undef LOCTEXT_NAMESPACE
\ No newline at end of file
diff --git a/Engine/Plugins/Experimental/Enterprise/DatasmithInterchange/Source/DatasmithInterchange/Private/InterchangeDatasmithLevelPipeline.cpp b/Engine/Plugins/Experimental/Enterprise/DatasmithInterchange/Source/DatasmithInterchange/Private/InterchangeDatasmithLevelPipeline.cpp
--- a/Engine/Plugins/Experimental/Enterprise/DatasmithInterchange/Source/DatasmithInterchange/Private/InterchangeDatasmithLevelPipeline.cpp
+++ b/Engine/Plugins/Experimental/Enterprise/DatasmithInterchange/Source/DatasmithInterchange/Private/InterchangeDatasmithLevelPipeline.cpp
@@ -6,9 +6,6 @@
 #include "InterchangeDatasmithAreaLightNode.h"
 #include "InterchangeDatasmithUtils.h"
 
-#include "InterchangeDecalActorFactoryNode.h"
-#include "InterchangeDecalNode.h"
-
 #include "InterchangeLevelSequenceFactoryNode.h"
 #include "InterchangeAnimationTrackSetNode.h"
 #include "InterchangeMaterialFactoryNode.h"
@@ -153,11 +150,6 @@ UInterchangeActorFactoryNode* UInterchangeDatasmithLevelPipeline::CreateActorFac
 	{
 		return NewObject<UInterchangeDatasmithAreaLightFactoryNode>(BaseNodeContainer, NAME_None);
 	}
-	// Experimental - Move to Super class once finalized
-	else if (TranslatedAssetNode && TranslatedAssetNode->IsA<UInterchangeDecalNode>())
-	{
-		return NewObject<UInterchangeDecalActorFactoryNode>(BaseNodeContainer, NAME_None);
-	}
 	else
 	{
 		return Super::CreateActorFactoryNode(SceneNode, TranslatedAssetNode);
@@ -172,13 +164,6 @@ void UInterchangeDatasmithLevelPipeline::SetUpFactoryNode(UInterchangeActorFacto
 		ensure(AreaLightFactory);
 		SetupAreaLight(AreaLightFactory, AreaLightNode);
 	}
-	// Experimental - Move to Super class once finalized
-	else if (const UInterchangeDecalNode* DecalNode = Cast<UInterchangeDecalNode>(TranslatedAssetNode))
-	{
-		UInterchangeDecalActorFactoryNode* DecalActorFactory = Cast<UInterchangeDecalActorFactoryNode>(ActorFactoryNode);
-		ensure(DecalActorFactory);
-		SetupDecalActor(DecalActorFactory, DecalNode);
-	}
 	else
 	{
 		Super::SetUpFactoryNode(ActorFactoryNode, SceneNode, TranslatedAssetNode);
@@ -278,21 +263,6 @@ void UInterchangeDatasmithLevelPipeline::SetupAreaLight(UInterchangeDatasmithAre
 	APPLY_FACTORY_ATTRIBUTE(SpotlightOuterAngle, float);
 }
 
-void UInterchangeDatasmithLevelPipeline::SetupDecalActor(UInterchangeDecalActorFactoryNode* FactoryNode, const UInterchangeDecalNode* TranslatedNode) const
-{
-	APPLY_FACTORY_ATTRIBUTE(DecalSize, FVector);
-	APPLY_FACTORY_ATTRIBUTE(SortOrder, int32);
-	APPLY_FACTORY_ATTRIBUTE(DecalMaterialPathName, FString);
-
-	// If the path is not a valid object path then it is an Interchange Node UID (Decal Material Node to be specific).
-	if (!FPackageName::IsValidObjectPath(DecalMaterialPathName))
-	{
-		const FString MaterialFactoryUid = UInterchangeFactoryBaseNode::BuildFactoryNodeUid(DecalMaterialPathName);
-		FactoryNode->SetCustomDecalMaterialPathName(MaterialFactoryUid);
-		FactoryNode->AddFactoryDependencyUid(MaterialFactoryUid);
-	}
-}
-
 #undef APPLY_FACTORY_ATTRIBUTE
 #undef APPLY_FACTORY_ATTRIBUTE_WITH_VALIDATION
 #undef APPLY_FACTORY_ATTRIBUTE_WITH_MODIFIER
\ No newline at end of file
diff --git a/Engine/Plugins/Experimental/Enterprise/DatasmithInterchange/Source/DatasmithInterchange/Private/InterchangeDatasmithLevelPipeline.h b/Engine/Plugins/Experimental/Enterprise/DatasmithInterchange/Source/DatasmithInterchange/Private/InterchangeDatasmithLevelPipeline.h
--- a/Engine/Plugins/Experimental/Enterprise/DatasmithInterchange/Source/DatasmithInterchange/Private/InterchangeDatasmithLevelPipeline.h
+++ b/Engine/Plugins/Experimental/Enterprise/DatasmithInterchange/Source/DatasmithInterchange/Private/InterchangeDatasmithLevelPipeline.h
@@ -32,5 +32,6 @@ protected:
 private:
 
 	void SetupAreaLight(UInterchangeDatasmithAreaLightFactoryNode* AreaLightFactoryNode, const UInterchangeDatasmithAreaLightNode* AreaLightNode) const;
-	void SetupDecalActor(UInterchangeDecalActorFactoryNode* DecalActorFactoryNode, const UInterchangeDecalNode* DecalNode) const;
+
+	friend class UInterchangeDatasmithPipeline;
 };
\ No newline at end of file
diff --git a/Engine/Plugins/Experimental/Enterprise/DatasmithInterchange/Source/DatasmithInterchange/Private/InterchangeDatasmithPipeline.cpp b/Engine/Plugins/Experimental/Enterprise/DatasmithInterchange/Source/DatasmithInterchange/Private/InterchangeDatasmithPipeline.cpp
--- a/Engine/Plugins/Experimental/Enterprise/DatasmithInterchange/Source/DatasmithInterchange/Private/InterchangeDatasmithPipeline.cpp
+++ b/Engine/Plugins/Experimental/Enterprise/DatasmithInterchange/Source/DatasmithInterchange/Private/InterchangeDatasmithPipeline.cpp
@@ -5,17 +5,19 @@
 #include "InterchangeDatasmithAreaLightNode.h"
 #include "InterchangeDatasmithAreaLightFactoryNode.h"
 #include "InterchangeDatasmithLevelPipeline.h"
+#include "InterchangeDatasmithLog.h"
 #include "InterchangeDatasmithMaterialPipeline.h"
-#include "InterchangeDatasmithSceneNode.h"
-#include "InterchangeDatasmithSceneFactoryNode.h"
 #include "InterchangeDatasmithStaticMeshPipeline.h"
-#include "InterchangeDatasmithTexturePipeline.h"
+//#include "InterchangeDatasmithTexturePipeline.h"
+#include "InterchangeDatasmithTextureData.h"
 #include "InterchangeDatasmithUtils.h"
 
-#include "InterchangeLevelSequenceFactoryNode.h"
 #include "InterchangeGenericAnimationPipeline.h"
 #include "InterchangeGenericScenesPipeline.h"
+#include "InterchangeGenericTexturePipeline.h"
+#include "InterchangeLevelSequenceFactoryNode.h"
 #include "InterchangeMaterialFactoryNode.h"
+#include "InterchangeSceneImportAssetFactoryNode.h"
 #include "InterchangeSceneVariantSetsFactoryNode.h"
 #include "InterchangeStaticMeshFactoryNode.h"
 #include "InterchangeTexture2DNode.h"
@@ -25,8 +27,6 @@
 
 #include "ExternalSource.h"
 #include "DatasmithAreaLightActor.h"
-#include "DatasmithScene.h"
-#include "DatasmithSceneXmlWriter.h"
 
 #include "StaticMeshAttributes.h"
 #include "Engine/StaticMeshActor.h"
@@ -107,7 +107,6 @@ UInterchangeDatasmithPipeline::UInterchangeDatasmithPipeline()
 	CommonMeshesProperties = CreateDefaultSubobject<UInterchangeGenericCommonMeshesProperties>("CommonMeshesProperties");
 	CommonMeshesProperties->bBakeMeshes = false;
 
-	TexturePipeline = CreateDefaultSubobject<UInterchangeDatasmithTexturePipeline>("DatasmithTexturePipeline");
 	MaterialPipeline = CreateDefaultSubobject<UInterchangeDatasmithMaterialPipeline>("DatasmithMaterialPipeline");
 	MeshPipeline = CreateDefaultSubobject<UInterchangeDatasmithStaticMeshPipeline>("DatasmithMeshPipeline");
 	LevelPipeline = CreateDefaultSubobject<UInterchangeDatasmithLevelPipeline>("DatasmithLevelPipeline");
@@ -126,11 +125,8 @@ void UInterchangeDatasmithPipeline::ExecutePipeline(UInterchangeBaseNodeContaine
 
 	BaseNodeContainer = InBaseNodeContainer;
 
-	TArray<UInterchangeDatasmithSceneNode*> DatasmithSceneNodes = NodeUtils::GetNodes<UInterchangeDatasmithSceneNode>(BaseNodeContainer);
-	if (!ensure(DatasmithSceneNodes.Num() == 1))
-	{
-		return;
-	}
+	ensure(BaseNodeContainer);
+	ensure(Results);
 
 	auto ExecutePreImportPipelineFunc = [this, &SourceDatas, &ContentBasePath](UInterchangePipelineBase* Pipeline)
 	{
@@ -141,22 +137,25 @@ void UInterchangeDatasmithPipeline::ExecutePipeline(UInterchangeBaseNodeContaine
 		}
 	};
 
-	ensure(Results);
-	ExecutePreImportPipelineFunc(TexturePipeline);
 	ExecutePreImportPipelineFunc(MaterialPipeline);
+	for (UInterchangeTextureFactoryNode* TextureFactoryNode : NodeUtils::GetNodes<UInterchangeTextureFactoryNode>(BaseNodeContainer))
+	{
+		PreImportTextureFactoryNode(TextureFactoryNode);
+	}
+
 	ExecutePreImportPipelineFunc(MeshPipeline);
 	ExecutePreImportPipelineFunc(LevelPipeline);
 	ExecutePreImportPipelineFunc(AnimationPipeline);
 
-	const UInterchangeDatasmithSceneNode* DatasmithSceneNode = DatasmithSceneNodes[0];
-	const FString PackageSubPath = DatasmithSceneNode->GetDisplayLabel();
+	const FString PackageSubPath = FPaths::GetBaseFilename(SourceDatas[0]->GetFilename());
 
-	TArray<FString> DependenciesUids;
+#if WITH_EDITORONLY_DATA
+	LevelPipeline->SceneImportFactoryNode->SetCustomSubPath(PackageSubPath);
+#endif
 
 	// Textures
 	for (UInterchangeTextureFactoryNode* TextureFactoryNode : NodeUtils::GetNodes<UInterchangeTextureFactoryNode>(BaseNodeContainer))
 	{
-		DependenciesUids.Add(TextureFactoryNode->GetUniqueID());
 		TextureFactoryNode->SetCustomSubPath(FPaths::Combine(PackageSubPath, "Textures"));
 		TextureFactoryNode->SetEnabled(true);
 	}
@@ -164,7 +163,6 @@ void UInterchangeDatasmithPipeline::ExecutePipeline(UInterchangeBaseNodeContaine
 	// Materials
 	for (UInterchangeBaseMaterialFactoryNode* MaterialFactoryNode : NodeUtils::GetNodes<UInterchangeBaseMaterialFactoryNode>(BaseNodeContainer))
 	{
-		DependenciesUids.Add(MaterialFactoryNode->GetUniqueID());
 		if (bCreateMaterialReferencesFolders)
 		{
 			if (MaterialFactoryNode->IsA<UInterchangeMaterialFactoryNode>())
@@ -190,7 +188,6 @@ void UInterchangeDatasmithPipeline::ExecutePipeline(UInterchangeBaseNodeContaine
 	// StaticMeshes
 	for (UInterchangeStaticMeshFactoryNode* StaticMeshFactoryNode : NodeUtils::GetNodes<UInterchangeStaticMeshFactoryNode>(BaseNodeContainer))
 	{
-		DependenciesUids.Add(StaticMeshFactoryNode->GetUniqueID());
 		StaticMeshFactoryNode->SetCustomSubPath(FPaths::Combine(PackageSubPath, "Geometries"));
 		StaticMeshFactoryNode->SetEnabled(true);
 	}
@@ -198,7 +195,6 @@ void UInterchangeDatasmithPipeline::ExecutePipeline(UInterchangeBaseNodeContaine
 	// LevelSequences
 	for (UInterchangeLevelSequenceFactoryNode* AnimationTrackSetFactoryNode : NodeUtils::GetNodes<UInterchangeLevelSequenceFactoryNode>(BaseNodeContainer))
 	{
-		DependenciesUids.Add(AnimationTrackSetFactoryNode->GetUniqueID());
 		AnimationTrackSetFactoryNode->SetCustomSubPath(FPaths::Combine(PackageSubPath, "Animations"));
 		AnimationTrackSetFactoryNode->SetEnabled(true);
 	}
@@ -206,44 +202,15 @@ void UInterchangeDatasmithPipeline::ExecutePipeline(UInterchangeBaseNodeContaine
 	// LevelVariantSets
 	for (UInterchangeSceneVariantSetsFactoryNode* LevelVariantSetFactoryNode : NodeUtils::GetNodes<UInterchangeSceneVariantSetsFactoryNode>(BaseNodeContainer))
 	{
-		DependenciesUids.Add(LevelVariantSetFactoryNode->GetUniqueID());
 		LevelVariantSetFactoryNode->SetCustomSubPath(FPaths::Combine(PackageSubPath, "Variants"));
 		LevelVariantSetFactoryNode->SetEnabled(true);
 	}
-
-	// Datasmith Scene
-	{
-		const FString DatasmithSceneUid = UInterchangeFactoryBaseNode::BuildFactoryNodeUid(DatasmithSceneNode->GetUniqueID());
-		const FString DisplayLabel = DatasmithSceneNode->GetDisplayLabel();
-		UInterchangeDatasmithSceneFactoryNode* DatasmithSceneFactoryNode = NewObject<UInterchangeDatasmithSceneFactoryNode>(BaseNodeContainer, NAME_None);
-		if (!ensure(DatasmithSceneFactoryNode))
-		{
-			return;
-		}
-
-		DatasmithSceneFactoryNode->InitializeDatasmithFactorySceneNode(DatasmithSceneUid, DisplayLabel, UDatasmithScene::StaticClass()->GetName());
-		DatasmithSceneFactoryNode->SetCustomSubPath(PackageSubPath);
-		DatasmithSceneFactoryNode->AddTargetNodeUid(DatasmithSceneNode->GetUniqueID());
-		DatasmithSceneNode->AddTargetNodeUid(DatasmithSceneFactoryNode->GetUniqueID());
-
-		for (const FString& Uid : DependenciesUids)
-		{
-			DatasmithSceneFactoryNode->AddFactoryDependencyUid(Uid);
-		}
-
-		BaseNodeContainer->AddNode(DatasmithSceneFactoryNode);
-	}
 }
 
 void UInterchangeDatasmithPipeline::ExecutePostImportPipeline(const UInterchangeBaseNodeContainer* InBaseNodeContainer, const FString& NodeKey, UObject* CreatedAsset, bool bIsAReimport)
 {
 	TRACE_CPUPROFILER_EVENT_SCOPE(UInterchangeDatasmithPipeline::ExecutePostImportPipeline);
 
-	if (TexturePipeline)
-	{
-		TexturePipeline->ScriptedExecutePostImportPipeline(InBaseNodeContainer, NodeKey, CreatedAsset, bIsAReimport);
-	}
-
 	if (MaterialPipeline)
 	{
 		MaterialPipeline->ScriptedExecutePostImportPipeline(InBaseNodeContainer, NodeKey, CreatedAsset, bIsAReimport);
@@ -263,91 +230,148 @@ void UInterchangeDatasmithPipeline::ExecutePostImportPipeline(const UInterchange
 	{
 		AnimationPipeline->ScriptedExecutePostImportPipeline(InBaseNodeContainer, NodeKey, CreatedAsset, bIsAReimport);
 	}
+}
+
+void UInterchangeDatasmithPipeline::PreImportTextureFactoryNode(UInterchangeTextureFactoryNode* TextureFactoryNode) const
+{
+	using namespace UE::DatasmithInterchange;
+
+	TArray<FString> TargetNodes;
+	TextureFactoryNode->GetTargetNodeUids(TargetNodes);
+	if (TargetNodes.Num() == 0)
+	{
+		return;
+	}
 
-	if (UDatasmithScene* DatasmithSceneAsset = Cast<UDatasmithScene>(CreatedAsset))
+	const UInterchangeBaseNode* TargetNode = BaseNodeContainer->GetNode(TargetNodes[0]);
+	if (!FInterchangeDatasmithTextureData::HasData(TargetNode))
 	{
-		PostImportDatasmithSceneAsset(*DatasmithSceneAsset);
+		return;
 	}
-#if WITH_EDITORONLY_DATA
-	if (bDeleteInvalidMeshes)
+
+	FInterchangeDatasmithTextureDataConst TextureData(TargetNode);
+	TOptional< bool > bLocalFlipNormalMapGreenChannel;
+	TOptional< TextureMipGenSettings > MipGenSettings;
+	TOptional< TextureGroup > LODGroup;
+	TOptional< TextureCompressionSettings > CompressionSettings;
+
+	// Make sure to set the proper LODGroup as it's used to determine the CompressionSettings when using TEXTUREGROUP_WorldNormalMap
+	EDatasmithTextureMode TextureMode;
+	if (TextureData.GetCustomTextureMode(TextureMode))
 	{
-		if (UStaticMesh* StaticMesh = Cast<UStaticMesh>(CreatedAsset))
+		switch (TextureMode)
 		{
-			if (!UE::Interchange::StaticMeshUtils::IsMeshValid(StaticMesh))
+		case EDatasmithTextureMode::Diffuse:
+			LODGroup = TEXTUREGROUP_World;
+			break;
+		case EDatasmithTextureMode::Specular:
+			LODGroup = TEXTUREGROUP_WorldSpecular;
+			break;
+		case EDatasmithTextureMode::Bump:
+		case EDatasmithTextureMode::Normal:
+			LODGroup = TEXTUREGROUP_WorldNormalMap;
+			CompressionSettings = TC_Normalmap;
+			break;
+		case EDatasmithTextureMode::NormalGreenInv:
+			LODGroup = TEXTUREGROUP_WorldNormalMap;
+			CompressionSettings = TC_Normalmap;
+			bLocalFlipNormalMapGreenChannel = true;
+			break;
+		}
+	}
+
+	const TOptional< float > RGBCurve = [&TextureData]() -> TOptional< float >
+		{
+			float ElementRGBCurve;
+
+			if (TextureData.GetCustomRGBCurve(ElementRGBCurve)
+				&& FMath::IsNearlyEqual(ElementRGBCurve, 1.0f) == false
+				&& ElementRGBCurve > 0.f)
 			{
-				UInterchangeResultWarning_Generic* InvalidMeshWarning = AddMessage<UInterchangeResultWarning_Generic>();
-				InvalidMeshWarning->Text = FText::Format(
-					LOCTEXT("InvalidStaticMesh", "Static Mesh {0} contains only degenerate or empty triangles."),
-					FText::FromString(StaticMesh->GetName())
-				);
-
-				if (UInterchangeStaticMeshFactoryNode* FactoryNode = Cast<UInterchangeStaticMeshFactoryNode>(InBaseNodeContainer->GetFactoryNode(NodeKey)))
-				{
-					FactoryNode->SetCustomReferenceObject(FSoftObjectPath());
-					InvalidStaticMeshFactoryUids.Emplace(FactoryNode->GetUniqueID());
-				}
-
-				CreatedAsset->Rename(nullptr, GetTransientPackage(), REN_NonTransactional | REN_DontCreateRedirectors);
-				if (CreatedAsset->IsRooted())
-				{
-					CreatedAsset->RemoveFromRoot();
-				}
-				CreatedAsset->ClearFlags(RF_Public | RF_Standalone);
-				CreatedAsset->MarkAsGarbage();
+				return ElementRGBCurve;
 			}
-		}
 
-		// Clear up the reference 
-		if (AStaticMeshActor* StaticMeshActor = Cast<AStaticMeshActor>(CreatedAsset))
+			return {};
+		}();
+
+		static_assert(TextureAddress::TA_Wrap == (int)EDatasmithTextureAddress::Wrap && TextureAddress::TA_Mirror == (int)EDatasmithTextureAddress::Mirror, "Texture Address enum doesn't match!");
+
+		TOptional< TextureFilter > TexFilter;
+		EDatasmithTextureFilter TextureFilter;
+		if (TextureData.GetCustomTextureFilter(TextureFilter))
 		{
-			if (UInterchangeMeshActorFactoryNode* FactoryNode = Cast<UInterchangeMeshActorFactoryNode>(InBaseNodeContainer->GetFactoryNode(NodeKey)))
+			switch (TextureFilter)
 			{
-				TArray<FString> FactoryDependencies;
-				FactoryNode->GetFactoryDependencies(FactoryDependencies);
-				for (const FString& FactoryDependency : FactoryDependencies)
-				{
-					if (InvalidStaticMeshFactoryUids.Contains(FactoryDependency))
-					{
-						if (UStaticMeshComponent* StaticMeshComponent = StaticMeshActor->GetStaticMeshComponent())
-						{
-							StaticMeshComponent->SetStaticMesh(nullptr);
-						}
-					}
-				}
+			case EDatasmithTextureFilter::Nearest:
+				TexFilter = TextureFilter::TF_Nearest;
+				break;
+			case EDatasmithTextureFilter::Bilinear:
+				TexFilter = TextureFilter::TF_Bilinear;
+				break;
+			case EDatasmithTextureFilter::Trilinear:
+				TexFilter = TextureFilter::TF_Trilinear;
+				break;
 			}
 		}
-	}
-#endif
-}
 
-void UInterchangeDatasmithPipeline::PostImportDatasmithSceneAsset(UDatasmithScene& DatasmithSceneAsset)
-{
-	using namespace UE::DatasmithInterchange;
+		TOptional< bool > bSrgb;
+		EDatasmithColorSpace ColorSpace;
+		if (TextureData.GetCustomSRGB(ColorSpace))
+		{
+			if (ColorSpace == EDatasmithColorSpace::sRGB)
+			{
+				bSrgb = true;
+			}
+			else if (ColorSpace == EDatasmithColorSpace::Linear)
+			{
+				bSrgb = false;
+			}
+		}
 
-#if WITH_EDITORONLY_DATA
-	TArray<UInterchangeDatasmithSceneNode*> DatasmithSceneNodes = NodeUtils::GetNodes<UInterchangeDatasmithSceneNode>(BaseNodeContainer);
-	if (!ensure(DatasmithSceneNodes.Num() == 1))
-	{
-		// TODO: Warn more than one Datasmith scene asset
-		return;
-	}
+		EDatasmithTextureAddress AddressX;
+		EDatasmithTextureAddress AddressY;
+		UInterchangeTexture2DFactoryNode* Texture2DFactoryNode = Cast<UInterchangeTexture2DFactoryNode>(TextureFactoryNode);
+		if (Texture2DFactoryNode
+			&& TextureData.GetCustomTextureAddressX(AddressX)
+			&& TextureData.GetCustomTextureAddressY(AddressY))
+		{
+			Texture2DFactoryNode->SetCustomAddressX((TextureAddress)AddressX);
+			Texture2DFactoryNode->SetCustomAddressY((TextureAddress)AddressY);
+		}
 
-	TArray< uint8 > Bytes;
-	FMemoryWriter MemoryWriter(Bytes, true);
+		if (bSrgb.IsSet())
+		{
+			TextureFactoryNode->SetCustomSRGB(bSrgb.GetValue());
+		}
 
-	FDatasmithSceneXmlWriter DatasmithSceneXmlWriter;
-	DatasmithSceneXmlWriter.Serialize(DatasmithSceneNodes[0]->DatasmithScene.ToSharedRef(), MemoryWriter);
+		if (bLocalFlipNormalMapGreenChannel.IsSet())
+		{
+			TextureFactoryNode->SetCustombFlipGreenChannel(bLocalFlipNormalMapGreenChannel.GetValue());
+		}
 
-	DatasmithSceneAsset.DatasmithSceneBulkData.Lock(LOCK_READ_WRITE);
+		if (MipGenSettings.IsSet())
+		{
+			TextureFactoryNode->SetCustomMipGenSettings(MipGenSettings.GetValue());
+		}
 
-	uint8* Dest = reinterpret_cast<uint8*>(DatasmithSceneAsset.DatasmithSceneBulkData.Realloc(Bytes.Num()));
+		if (LODGroup.IsSet())
+		{
+			TextureFactoryNode->SetCustomLODGroup(LODGroup.GetValue());
+		}
 
-	FPlatformMemory::Memcpy(Dest, Bytes.GetData(), Bytes.Num());
+		if (CompressionSettings.IsSet())
+		{
+			TextureFactoryNode->SetCustomCompressionSettings(CompressionSettings.GetValue());
+		}
 
-	DatasmithSceneAsset.DatasmithSceneBulkData.Unlock();
+		if (RGBCurve.IsSet())
+		{
+			TextureFactoryNode->SetCustomAdjustRGBCurve(RGBCurve.GetValue());
+		}
 
-	// Todo: Fill up imported DatasmithScene with created assets
-#endif
+		if (TexFilter.IsSet())
+		{
+			TextureFactoryNode->SetCustomFilter(TexFilter.GetValue());
+		}
 }
-
 #undef LOCTEXT_NAMESPACE
\ No newline at end of file
diff --git a/Engine/Plugins/Experimental/Enterprise/DatasmithInterchange/Source/DatasmithInterchange/Private/InterchangeDatasmithSceneFactory.cpp b/Engine/Plugins/Experimental/Enterprise/DatasmithInterchange/Source/DatasmithInterchange/Private/InterchangeDatasmithSceneFactory.cpp
deleted file mode 100644
diff --git a/Engine/Plugins/Experimental/Enterprise/DatasmithInterchange/Source/DatasmithInterchange/Private/InterchangeDatasmithSceneFactoryNode.cpp b/Engine/Plugins/Experimental/Enterprise/DatasmithInterchange/Source/DatasmithInterchange/Private/InterchangeDatasmithSceneFactoryNode.cpp
deleted file mode 100644
diff --git a/Engine/Plugins/Experimental/Enterprise/DatasmithInterchange/Source/DatasmithInterchange/Private/InterchangeDatasmithSceneNode.cpp b/Engine/Plugins/Experimental/Enterprise/DatasmithInterchange/Source/DatasmithInterchange/Private/InterchangeDatasmithSceneNode.cpp
deleted file mode 100644
diff --git a/Engine/Plugins/Experimental/Enterprise/DatasmithInterchange/Source/DatasmithInterchange/Private/InterchangeDatasmithTexturePipeline.cpp b/Engine/Plugins/Experimental/Enterprise/DatasmithInterchange/Source/DatasmithInterchange/Private/InterchangeDatasmithTexturePipeline.cpp
deleted file mode 100644
diff --git a/Engine/Plugins/Experimental/Enterprise/DatasmithInterchange/Source/DatasmithInterchange/Private/InterchangeDatasmithTexturePipeline.h b/Engine/Plugins/Experimental/Enterprise/DatasmithInterchange/Source/DatasmithInterchange/Private/InterchangeDatasmithTexturePipeline.h
deleted file mode 100644
diff --git a/Engine/Plugins/Experimental/Enterprise/DatasmithInterchange/Source/DatasmithInterchange/Private/InterchangeDatasmithTranslator.cpp b/Engine/Plugins/Experimental/Enterprise/DatasmithInterchange/Source/DatasmithInterchange/Private/InterchangeDatasmithTranslator.cpp
--- a/Engine/Plugins/Experimental/Enterprise/DatasmithInterchange/Source/DatasmithInterchange/Private/InterchangeDatasmithTranslator.cpp
+++ b/Engine/Plugins/Experimental/Enterprise/DatasmithInterchange/Source/DatasmithInterchange/Private/InterchangeDatasmithTranslator.cpp
@@ -5,7 +5,6 @@
 #include "InterchangeDatasmithAreaLightNode.h"
 #include "InterchangeDatasmithLog.h"
 #include "InterchangeDatasmithMaterialNode.h"
-#include "InterchangeDatasmithSceneNode.h"
 #include "InterchangeDatasmithTextureData.h"
 #include "InterchangeDatasmithUtils.h"
 
@@ -39,9 +38,101 @@
 #include "Misc/App.h"
 #include "Misc/PackageName.h"
 
+#if WITH_EDITOR
+#include "DesktopPlatformModule.h"
+#include "Dialogs/DlgPickPath.h"
+#include "IDesktopPlatform.h"
+#include "Interfaces/IMainFrameModule.h"
+#include "ObjectTools.h"
+#include "Styling/SlateIconFinder.h"
+#include "UI/DatasmithImportOptionsWindow.h"
+#endif //WITH_EDITOR
 
 #define LOCTEXT_NAMESPACE "DatasmithInterchange"
 
+namespace UE::Interchange::Datasmith
+{
+#if WITH_EDITOR
+	bool DisplayOptionsDialog(IDatasmithTranslator& Translator)
+	{
+		TArray<TObjectPtr<UDatasmithOptionsBase>> ImportOptions;
+		Translator.GetSceneImportOptions(ImportOptions);
+
+		if (ImportOptions.Num() == 0)
+		{
+			return true;
+		}
+
+		const FString FilePath = Translator.GetSource().GetSourceFile();
+
+		TSharedPtr<SWindow> ParentWindow;
+
+		if (FModuleManager::Get().IsModuleLoaded("MainFrame"))
+		{
+			IMainFrameModule& MainFrame = FModuleManager::LoadModuleChecked<IMainFrameModule>("MainFrame");
+			ParentWindow = MainFrame.GetParentWindow();
+		}
+
+		TArray<UObject*> Options;
+		Options.SetNum(ImportOptions.Num());
+		for (int32 Index = 0; Index < ImportOptions.Num(); ++Index)
+		{
+			Options[Index] = ImportOptions[Index];
+		}
+
+		TSharedRef<SWindow> Window = SNew(SWindow)
+			.Title(LOCTEXT("DatasmithImportSettingsTitle", "Datasmith Import Options"))
+			.SizingRule(ESizingRule::Autosized);
+
+		float SceneVersion = FDatasmithUtils::GetDatasmithFormatVersionAsFloat();
+		FString FileSDKVersion = FDatasmithUtils::GetEnterpriseVersionAsString();
+
+		TSharedPtr<SDatasmithOptionsWindow> OptionsWindow;
+		Window->SetContent
+		(
+			SAssignNew(OptionsWindow, SDatasmithOptionsWindow)
+			.ImportOptions(Options)
+			.WidgetWindow(Window)
+			// note: Spacing in text below is intentional for text alignment
+			.FileNameText(FText::Format(LOCTEXT("DatasmithImportSettingsFileName", "  Import File  :    {0}"), FText::FromString(FPaths::GetCleanFilename(FilePath))))
+			.FilePathText(FText::FromString(FilePath))
+			.FileFormatVersion(SceneVersion)
+			.FileSDKVersion(FText::FromString(FileSDKVersion))
+//			.PackagePathText(FText::Format(LOCTEXT("DatasmithImportSettingsPackagePath", "  Import To   :    {0}"), FText::FromString(PackagePath)))
+			.ProceedButtonLabel(LOCTEXT("DatasmithOptionWindow_ImportCurLevel", "Import"))
+			.ProceedButtonTooltip(LOCTEXT("DatasmithOptionWindow_ImportCurLevel_ToolTip", "Import the file through Interchange and add to the current Level"))
+			.CancelButtonLabel(LOCTEXT("DatasmithOptionWindow_Cancel", "Cancel"))
+			.CancelButtonTooltip(LOCTEXT("DatasmithOptionWindow_Cancel_ToolTip", "Cancel importing this file"))
+			.MinDetailHeight(320.f)
+			.MinDetailWidth(450.f)
+		);
+
+		FSlateApplication::Get().AddModalWindow(Window, ParentWindow, false);
+
+		return OptionsWindow->ShouldImport();
+	}
+#endif
+}
+
+TArray<FString> UInterchangeDatasmithTranslator::GetSupportedFormats() const
+{
+	const TArray<FString> DatasmithFormats = FDatasmithTranslatorManager::Get().GetSupportedFormats();
+	TArray<FString> Formats;
+	Formats.Reserve(DatasmithFormats.Num() - 1);
+
+	for (const FString& Format : DatasmithFormats)
+	{
+		if (Format.Contains(TEXT("gltf")) || Format.Contains(TEXT("glb")))
+		{
+			continue;
+		}
+
+		Formats.Add(Format);
+	}
+
+	return Formats;
+}
+
 bool UInterchangeDatasmithTranslator::CanImportSourceData(const UInterchangeSourceData* InSourceData) const
 {
 	using namespace UE::DatasmithImporter;
@@ -69,38 +160,51 @@ bool UInterchangeDatasmithTranslator::Translate(UInterchangeBaseNodeContainer& B
 	FString FilePath = FPaths::ConvertRelativePathToFull(SourceData->GetFilename());
 	FileName = FPaths::GetCleanFilename(FilePath);
 	const FSourceUri FileNameUri = FSourceUri::FromFilePath(FilePath);
-	const_cast<UInterchangeDatasmithTranslator*>(this)->LoadedExternalSource = IExternalSourceModule::GetOrCreateExternalSource(FileNameUri);
+	LoadedExternalSource = IExternalSourceModule::GetOrCreateExternalSource(FileNameUri);
 
 	if (!LoadedExternalSource.IsValid() || !LoadedExternalSource->IsAvailable())
 	{
 		return false;
 	}
 
-	UClass* Class = UInterchangeDatasmithSceneNode::StaticClass();
-	if (!ensure(Class))
-	{
-		return false;
-	}
-
-	// Temporary: Update the tessellation options of the associated translator
+	// Temporary: Update the extra options of the associated translator
+#if WITH_EDITOR
 	{
 		const TSharedPtr<IDatasmithTranslator>& DatasmithTranslator = LoadedExternalSource->GetAssetTranslator();
-
-		const FString OptionFilePath = BuildConfigFilePath(FilePath);
-
-		if (DatasmithTranslator.IsValid() && FPaths::FileExists(OptionFilePath))
+		if (!DatasmithTranslator)
 		{
-			TArray<TObjectPtr<UDatasmithOptionsBase>> ImportOptions;
-			DatasmithTranslator->GetSceneImportOptions(ImportOptions);
+			return false;
+		}
 
-			for (TObjectPtr<UDatasmithOptionsBase>& Option : ImportOptions)
+		bool bShouldImport = true;
+		if (!IsInGameThread())
+		{
+			TSharedRef<TPromise<bool>, ESPMode::ThreadSafe> Promise = MakeShareable(new TPromise<bool>());
+			TFunction<void()> PromiseKeeper = [&DatasmithTranslator, &Promise]() -> void
 			{
-				Option->LoadConfig(nullptr, *OptionFilePath);
-			}
+				bool bShouldImport = UE::Interchange::Datasmith::DisplayOptionsDialog(*DatasmithTranslator);
+				Promise->SetValue(bShouldImport);
+			};
 
-			DatasmithTranslator->SetSceneImportOptions(ImportOptions);
+			AsyncTask(ENamedThreads::GameThread, MoveTemp(PromiseKeeper));
+
+			TFuture<bool> Future = Promise->GetFuture();
+			
+			Future.Wait();
+			
+			bShouldImport = Future.Get();
+		}
+		else
+		{
+			bShouldImport = UE::Interchange::Datasmith::DisplayOptionsDialog(*DatasmithTranslator);
+		}
+
+		if (!bShouldImport)
+		{
+			return false;
 		}
 	}
+#endif
 
 	StartTime = FPlatformTime::Cycles64();
 	FPaths::NormalizeFilename(FilePath);
@@ -112,24 +216,6 @@ bool UInterchangeDatasmithTranslator::Translate(UInterchangeBaseNodeContainer& B
 		return false;
 	}
 
-	// Datasmith Scene Node
-	{
-		FString DisplayLabel = DatasmithScene->GetName();
-		FString NodeUID(NodeUtils::DatasmithScenePrefix + FilePath);
-		UInterchangeDatasmithSceneNode* DatasmithSceneNode = NewObject<UInterchangeDatasmithSceneNode>(&BaseNodeContainer, Class);
-		if (!ensure(DatasmithSceneNode))
-		{
-			return false;
-		}
-		DatasmithSceneNode->InitializeNode(NodeUID, DisplayLabel, EInterchangeNodeContainerType::TranslatedAsset);
-
-		// Assigning those variables are our only way to pass the translated state to the part of the pipeline that does not use interchange yet.
-		DatasmithSceneNode->ExternalSource = LoadedExternalSource;
-		DatasmithSceneNode->DatasmithScene = DatasmithScene;
-
-		BaseNodeContainer.AddNode(DatasmithSceneNode);
-	}
-
 	// Texture Nodes
 	{
 		FDatasmithUniqueNameProvider TextureNameProvider;
@@ -320,6 +406,13 @@ bool UInterchangeDatasmithTranslator::Translate(UInterchangeBaseNodeContainer& B
 		VariantSetUtils::TranslateLevelVariantSets(LevelVariantSets, BaseNodeContainer);
 	}
 
+	// Log time spent to import incoming file in minutes and seconds
+	double ElapsedSeconds = FPlatformTime::ToSeconds64(FPlatformTime::Cycles64() - StartTime);
+
+	int ElapsedMin = int(ElapsedSeconds / 60.0);
+	ElapsedSeconds -= 60.0 * (double)ElapsedMin;
+	UE_LOG(LogInterchangeDatasmith, Log, TEXT("Translation of %s in[%d min %.3f s]"), *FileName, ElapsedMin, ElapsedSeconds);
+
 	return true;
 }
 
@@ -779,11 +872,5 @@ void UInterchangeDatasmithTranslator::ImportFinish()
 	UE_LOG(LogInterchangeDatasmith, Log, TEXT("Imported %s in [%d min %.3f s]"), *FileName, ElapsedMin, ElapsedSeconds);
 }
 
-FString UInterchangeDatasmithTranslator::BuildConfigFilePath(const FString& FilePath)
-{
-	const FString OptionFileName = FMD5::HashAnsiString(*(FPaths::ConvertRelativePathToFull(FilePath) + TEXT("_Config"))) + TEXT(".ini");
-	return  FPaths::Combine(FPlatformProcess::UserTempDir(), OptionFileName);
-}
-
 
 #undef LOCTEXT_NAMESPACE
\ No newline at end of file
diff --git a/Engine/Plugins/Experimental/Enterprise/DatasmithInterchange/Source/DatasmithInterchange/Public/InterchangeDatasmithPipeline.h b/Engine/Plugins/Experimental/Enterprise/DatasmithInterchange/Source/DatasmithInterchange/Public/InterchangeDatasmithPipeline.h
--- a/Engine/Plugins/Experimental/Enterprise/DatasmithInterchange/Source/DatasmithInterchange/Public/InterchangeDatasmithPipeline.h
+++ b/Engine/Plugins/Experimental/Enterprise/DatasmithInterchange/Source/DatasmithInterchange/Public/InterchangeDatasmithPipeline.h
@@ -18,7 +18,8 @@ class UInterchangeDatasmithLevelPipeline;
 class UInterchangeDatasmithMaterialPipeline;
 class UInterchangeDatasmithStaticMeshPipeline;
 class UInterchangeDatasmithSceneNode;
-class UInterchangeDatasmithTexturePipeline;
+class UInterchangeGenericTexturePipeline;
+class UInterchangeTextureFactoryNode;
 class UTexture;
 struct FDatasmithImportContext;
 
@@ -29,9 +30,6 @@ class DATASMITHINTERCHANGE_API UInterchangeDatasmithPipeline : public UInterchan
 
 	UInterchangeDatasmithPipeline();
 
-	UPROPERTY(VisibleAnywhere, Instanced, Category = "Textures")
-	TObjectPtr<UInterchangeDatasmithTexturePipeline> TexturePipeline;
-
 	UPROPERTY(VisibleAnywhere, Instanced, Category = "Materials")
 	TObjectPtr<UInterchangeDatasmithMaterialPipeline> MaterialPipeline;
 
@@ -69,10 +67,8 @@ protected:
 	}
 
 private:
-	UInterchangeBaseNodeContainer* BaseNodeContainer = nullptr;
-
-	// Fill up the UDatasmithScene with all the data its needs for DnD
-	void PostImportDatasmithSceneAsset(UDatasmithScene& DatasmithSceneAsset);
+	void PreImportTextureFactoryNode(UInterchangeTextureFactoryNode* TextureFactoryNode) const;
 
-	TSet<FString> InvalidStaticMeshFactoryUids;
+private:
+	UInterchangeBaseNodeContainer* BaseNodeContainer = nullptr;
 };
diff --git a/Engine/Plugins/Experimental/Enterprise/DatasmithInterchange/Source/DatasmithInterchange/Public/InterchangeDatasmithSceneFactory.h b/Engine/Plugins/Experimental/Enterprise/DatasmithInterchange/Source/DatasmithInterchange/Public/InterchangeDatasmithSceneFactory.h
deleted file mode 100644
diff --git a/Engine/Plugins/Experimental/Enterprise/DatasmithInterchange/Source/DatasmithInterchange/Public/InterchangeDatasmithSceneFactoryNode.h b/Engine/Plugins/Experimental/Enterprise/DatasmithInterchange/Source/DatasmithInterchange/Public/InterchangeDatasmithSceneFactoryNode.h
deleted file mode 100644
diff --git a/Engine/Plugins/Experimental/Enterprise/DatasmithInterchange/Source/DatasmithInterchange/Public/InterchangeDatasmithSceneNode.h b/Engine/Plugins/Experimental/Enterprise/DatasmithInterchange/Source/DatasmithInterchange/Public/InterchangeDatasmithSceneNode.h
deleted file mode 100644
diff --git a/Engine/Plugins/Experimental/Enterprise/DatasmithInterchange/Source/DatasmithInterchange/Public/InterchangeDatasmithTranslator.h b/Engine/Plugins/Experimental/Enterprise/DatasmithInterchange/Source/DatasmithInterchange/Public/InterchangeDatasmithTranslator.h
--- a/Engine/Plugins/Experimental/Enterprise/DatasmithInterchange/Source/DatasmithInterchange/Public/InterchangeDatasmithTranslator.h
+++ b/Engine/Plugins/Experimental/Enterprise/DatasmithInterchange/Source/DatasmithInterchange/Public/InterchangeDatasmithTranslator.h
@@ -66,8 +66,15 @@ public:
 
 	virtual bool Translate(UInterchangeBaseNodeContainer& BaseNodeContainer) const override;
 
-	virtual TArray<FString> GetSupportedFormats() const override { return TArray<FString>(); }
-	virtual EInterchangeTranslatorAssetType GetSupportedAssetTypes() const override { return EInterchangeTranslatorAssetType::Materials | EInterchangeTranslatorAssetType::Meshes;}
+	virtual TArray<FString> GetSupportedFormats() const override;
+	virtual EInterchangeTranslatorAssetType GetSupportedAssetTypes() const override
+	{ 
+		return EInterchangeTranslatorAssetType::Textures | EInterchangeTranslatorAssetType::Materials | EInterchangeTranslatorAssetType::Meshes;
+	}
+	virtual EInterchangeTranslatorType GetTranslatorType() const override
+	{
+		return EInterchangeTranslatorType::Scenes;
+	}
 
 	virtual void ReleaseSource() override
 	{
@@ -92,9 +99,6 @@ public:
 	virtual TFuture<TOptional<UE::Interchange::FVariantSetPayloadData>> GetVariantSetPayloadData(const FString& PayloadKey) const override;
 	/* IInterchangeVariantSetPayloadInterface End */
 
-	/** Returns a unique file path to */ 
-	static FString BuildConfigFilePath(const FString& FilePath);
-
 private:
 
 	void HandleDatasmithActor(UInterchangeBaseNodeContainer& BaseNodeContainer, const TSharedRef<IDatasmithActorElement>& ActorElement, const UInterchangeSceneNode* ParentNode) const;
@@ -105,7 +109,7 @@ private:
 
 	UInterchangeDecalNode* AddDecalNode(UInterchangeBaseNodeContainer& BaseNodeContainer, const TSharedRef<IDatasmithDecalActorElement>& DecalActor) const;
 
-	TSharedPtr<UE::DatasmithImporter::FExternalSource> LoadedExternalSource;
+	mutable TSharedPtr<UE::DatasmithImporter::FExternalSource> LoadedExternalSource;
 
 	mutable uint64 StartTime = 0;
 	mutable FString FileName;
diff --git a/Engine/Plugins/Experimental/Enterprise/DatasmithInterchange/Source/DatasmithInterchangeScripting/DatasmithInterchangeScripting.Build.cs b/Engine/Plugins/Experimental/Enterprise/DatasmithInterchange/Source/DatasmithInterchangeScripting/DatasmithInterchangeScripting.Build.cs
deleted file mode 100644
diff --git a/Engine/Plugins/Experimental/Enterprise/DatasmithInterchange/Source/DatasmithInterchangeScripting/Private/DatasmithInterchangeScripting.cpp b/Engine/Plugins/Experimental/Enterprise/DatasmithInterchange/Source/DatasmithInterchangeScripting/Private/DatasmithInterchangeScripting.cpp
deleted file mode 100644
diff --git a/Engine/Plugins/Experimental/Enterprise/DatasmithInterchange/Source/DatasmithInterchangeScripting/Private/DatasmithInterchangeScriptingModule.cpp b/Engine/Plugins/Experimental/Enterprise/DatasmithInterchange/Source/DatasmithInterchangeScripting/Private/DatasmithInterchangeScriptingModule.cpp
deleted file mode 100644
diff --git a/Engine/Plugins/Experimental/Enterprise/DatasmithInterchange/Source/DatasmithInterchangeScripting/Public/DatasmithInterchangeScripting.h b/Engine/Plugins/Experimental/Enterprise/DatasmithInterchange/Source/DatasmithInterchangeScripting/Public/DatasmithInterchangeScripting.h
deleted file mode 100644
diff --git a/Engine/Plugins/Interchange/Runtime/Source/Pipelines/Private/InterchangeGenericScenesPipeline.cpp b/Engine/Plugins/Interchange/Runtime/Source/Pipelines/Private/InterchangeGenericScenesPipeline.cpp
--- a/Engine/Plugins/Interchange/Runtime/Source/Pipelines/Private/InterchangeGenericScenesPipeline.cpp
+++ b/Engine/Plugins/Interchange/Runtime/Source/Pipelines/Private/InterchangeGenericScenesPipeline.cpp
@@ -14,6 +14,8 @@
 #include "InterchangeLightFactoryNode.h"
 #include "InterchangeMeshActorFactoryNode.h"
 #include "InterchangeMeshNode.h"
+#include "InterchangeDecalActorFactoryNode.h"
+#include "InterchangeDecalNode.h"
 #include "InterchangePipelineLog.h"
 #include "InterchangePipelineMeshesUtilities.h"
 #include "InterchangeSceneNode.h"
@@ -584,6 +586,10 @@ UInterchangeActorFactoryNode* UInterchangeGenericLevelPipeline::CreateActorFacto
 		{
 			return NewObject<UInterchangeDirectionalLightFactoryNode>(BaseNodeContainer, NAME_None);
 		}
+		else if (TranslatedAssetNode && TranslatedAssetNode->IsA<UInterchangeDecalNode>())
+		{
+			return NewObject<UInterchangeDecalActorFactoryNode>(BaseNodeContainer, NAME_None);
+		}
 	}
 
 	return NewObject<UInterchangeActorFactoryNode>(BaseNodeContainer, NAME_None);
@@ -872,6 +878,37 @@ void UInterchangeGenericLevelPipeline::SetUpFactoryNode(UInterchangeActorFactory
 			}
 		}
 	}
+	else if (const UInterchangeDecalNode* DecalNode = Cast<UInterchangeDecalNode>(TranslatedAssetNode))
+	{
+		UInterchangeDecalActorFactoryNode* DecalActorFactory = Cast<UInterchangeDecalActorFactoryNode>(ActorFactoryNode);
+		ensure(DecalActorFactory);
+
+		if (FVector DecalSize; DecalNode->GetCustomDecalSize(DecalSize))
+		{
+			DecalActorFactory->SetCustomDecalSize(DecalSize);
+		}
+
+		if (int32 SortOrder; DecalNode->GetCustomSortOrder(SortOrder))
+		{
+			DecalActorFactory->SetCustomSortOrder(SortOrder);
+		}
+		
+		bool bHasMaterialPathName = false;
+		FString DecalMaterialPathName;
+		if (DecalNode->GetCustomDecalMaterialPathName(DecalMaterialPathName))
+		{
+			DecalActorFactory->SetCustomDecalMaterialPathName(DecalMaterialPathName);
+			bHasMaterialPathName = true;
+		}
+
+		// If the path is not a valid object path then it is an Interchange Node UID (Decal Material Node to be specific).
+		if (bHasMaterialPathName && !FPackageName::IsValidObjectPath(DecalMaterialPathName))
+		{
+			const FString MaterialFactoryUid = UInterchangeFactoryBaseNode::BuildFactoryNodeUid(DecalMaterialPathName);
+			DecalActorFactory->SetCustomDecalMaterialPathName(MaterialFactoryUid);
+			DecalActorFactory->AddFactoryDependencyUid(MaterialFactoryUid);
+		}
+	}
 }
 
 void UInterchangeGenericLevelPipeline::ExecuteSceneVariantSetNodePreImport(const UInterchangeSceneVariantSetsNode& SceneVariantSetNode)
