Fixed several issues with Datasmith Interchange import
- Fixed import of Ies files as light profiles
- Fixed import of lights with attached Ies profiles
- Fixed import of level sequence animation
- Now differentiate asset import with level import

#virtualized

diff --git a/Engine/Plugins/Experimental/Enterprise/DatasmithInterchange/DatasmithInterchange.uplugin b/Engine/Plugins/Experimental/Enterprise/DatasmithInterchange/DatasmithInterchange.uplugin
--- a/Engine/Plugins/Experimental/Enterprise/DatasmithInterchange/DatasmithInterchange.uplugin
+++ b/Engine/Plugins/Experimental/Enterprise/DatasmithInterchange/DatasmithInterchange.uplugin
@@ -1,19 +1,20 @@
 {
-	"FileVersion" : 1,
-	"Version" : 1,
-	"VersionName" : "1.0",
-	"FriendlyName" : "Datasmith Interchange",
-	"Description" : "Interchange Importer for Datasmith.",
-	"Category" : "Importers",
-	"CreatedBy" : "Epic Games, Inc.",
-	"CreatedByURL" : "https://epicgames.com",
-	"DocsURL" : "",
-	"MarketplaceURL" : "",
-	"SupportURL" : "",
-	"EnabledByDefault" : false,
-	"CanContainContent" : true,
-	"IsBetaVersion" : true,
-	"Installed" : false,
+	"FileVersion": 1,
+	"Version": 1,
+	"VersionName": "1.0",
+	"FriendlyName": "Datasmith Interchange",
+	"Description": "Interchange Importer for Datasmith.",
+	"Category": "Importers",
+	"CreatedBy": "Epic Games, Inc.",
+	"CreatedByURL": "https://epicgames.com",
+	"DocsURL": "",
+	"MarketplaceURL": "",
+	"SupportURL": "",
+	"EnabledByDefault": false,
+	"CanContainContent": true,
+	"IsBetaVersion": false,
+	"IsExperimentalVersion": true,
+	"Installed": false,
 	"PlatformAllowList": [
 		"Win64",
 		"Linux",
diff --git a/Engine/Plugins/Experimental/Enterprise/DatasmithInterchange/Source/DatasmithInterchange/Private/InterchangeDatasmithAnimationUtils.cpp b/Engine/Plugins/Experimental/Enterprise/DatasmithInterchange/Source/DatasmithInterchange/Private/InterchangeDatasmithAnimationUtils.cpp
--- a/Engine/Plugins/Experimental/Enterprise/DatasmithInterchange/Source/DatasmithInterchange/Private/InterchangeDatasmithAnimationUtils.cpp
+++ b/Engine/Plugins/Experimental/Enterprise/DatasmithInterchange/Source/DatasmithInterchange/Private/InterchangeDatasmithAnimationUtils.cpp
@@ -392,6 +392,26 @@ namespace UE::DatasmithInterchange::AnimUtils
 		return false;
 	}
 
+	bool GetAnimationPayloadData(const IDatasmithBaseAnimationElement& AnimationElement, float FrameRate, EInterchangeAnimationPayLoadType PayLoadType, UE::Interchange::FAnimationPayloadData& PayLoadData)
+	{
+		switch (PayLoadType)
+		{
+		case EInterchangeAnimationPayLoadType::CURVE:
+		case EInterchangeAnimationPayLoadType::MORPHTARGETCURVE:
+			return GetAnimationPayloadData(AnimationElement, FrameRate, PayLoadData.Curves);
+
+		case EInterchangeAnimationPayLoadType::STEPCURVE:
+			return GetAnimationPayloadData(AnimationElement, FrameRate, PayLoadData.StepCurves);
+		
+		case EInterchangeAnimationPayLoadType::BAKED:
+		case EInterchangeAnimationPayLoadType::NONE:
+		default:
+			break;
+		}
+
+		return false;
+	}
+
 	void TranslateLevelSequences(TArray<TSharedPtr<IDatasmithLevelSequenceElement>>& LevelSequences, UInterchangeBaseNodeContainer& BaseNodeContainer, TMap<FString, UE::DatasmithInterchange::AnimUtils::FAnimationPayloadDesc>& AnimationPayLoadMapping)
 	{
 		FLevelSequenceImportHelper LevelSequenceHelper(BaseNodeContainer);
diff --git a/Engine/Plugins/Experimental/Enterprise/DatasmithInterchange/Source/DatasmithInterchange/Private/InterchangeDatasmithAreaLightNode.cpp b/Engine/Plugins/Experimental/Enterprise/DatasmithInterchange/Source/DatasmithInterchange/Private/InterchangeDatasmithAreaLightNode.cpp
--- a/Engine/Plugins/Experimental/Enterprise/DatasmithInterchange/Source/DatasmithInterchange/Private/InterchangeDatasmithAreaLightNode.cpp
+++ b/Engine/Plugins/Experimental/Enterprise/DatasmithInterchange/Source/DatasmithInterchange/Private/InterchangeDatasmithAreaLightNode.cpp
@@ -8,10 +8,6 @@ UE::Interchange::FAttributeKey UInterchangeDatasmithAreaLightNode::Macro_CustomL
 UE::Interchange::FAttributeKey UInterchangeDatasmithAreaLightNode::Macro_CustomLightShapeKey = UE::Interchange::FAttributeKey(TEXT("LightShape"));
 UE::Interchange::FAttributeKey UInterchangeDatasmithAreaLightNode::Macro_CustomDimensionsKey = UE::Interchange::FAttributeKey(TEXT("Dimensions"));
 UE::Interchange::FAttributeKey UInterchangeDatasmithAreaLightNode::Macro_CustomColorKey = UE::Interchange::FAttributeKey(TEXT("Color"));
-UE::Interchange::FAttributeKey UInterchangeDatasmithAreaLightNode::Macro_CustomIESTextureKey = UE::Interchange::FAttributeKey(TEXT("IESTexture"));
-UE::Interchange::FAttributeKey UInterchangeDatasmithAreaLightNode::Macro_CustomUseIESBrightnessKey = UE::Interchange::FAttributeKey(TEXT("UseIESBrightness"));
-UE::Interchange::FAttributeKey UInterchangeDatasmithAreaLightNode::Macro_CustomIESBrightnessScaleKey = UE::Interchange::FAttributeKey(TEXT("IESBrightnessScale"));
-UE::Interchange::FAttributeKey UInterchangeDatasmithAreaLightNode::Macro_CustomRotationKey = UE::Interchange::FAttributeKey(TEXT("Rotation"));
 UE::Interchange::FAttributeKey UInterchangeDatasmithAreaLightNode::Macro_CustomSourceRadiusKey = UE::Interchange::FAttributeKey(TEXT("SourceRadius"));
 UE::Interchange::FAttributeKey UInterchangeDatasmithAreaLightNode::Macro_CustomSourceLengthKey = UE::Interchange::FAttributeKey(TEXT("SourceLength"));
 UE::Interchange::FAttributeKey UInterchangeDatasmithAreaLightNode::Macro_CustomSpotlightInnerAngleKey = UE::Interchange::FAttributeKey(TEXT("SpotlightInnerAngle"));
@@ -65,46 +61,6 @@ bool UInterchangeDatasmithAreaLightNode::SetCustomColor(const FLinearColor& Attr
 	IMPLEMENT_NODE_ATTRIBUTE_SETTER_NODELEGATE(Color, FLinearColor);
 }
 
-bool UInterchangeDatasmithAreaLightNode::GetCustomIESTexture(FString& AttributeValue) const
-{
-	IMPLEMENT_NODE_ATTRIBUTE_GETTER(IESTexture, FString);
-}
-
-bool UInterchangeDatasmithAreaLightNode::SetCustomIESTexture(const FString& AttributeValue)
-{
-	IMPLEMENT_NODE_ATTRIBUTE_SETTER_NODELEGATE(IESTexture, FString);
-}
-
-bool UInterchangeDatasmithAreaLightNode::GetCustomUseIESBrightness(bool& AttributeValue) const
-{
-	IMPLEMENT_NODE_ATTRIBUTE_GETTER(UseIESBrightness, bool);
-}
-
-bool UInterchangeDatasmithAreaLightNode::SetCustomUseIESBrightness(const bool& AttributeValue, bool bAddApplyDelegate)
-{
-	IMPLEMENT_NODE_ATTRIBUTE_SETTER_NODELEGATE(UseIESBrightness, bool);
-}
-
-bool UInterchangeDatasmithAreaLightNode::GetCustomIESBrightnessScale(float& AttributeValue) const
-{
-	IMPLEMENT_NODE_ATTRIBUTE_GETTER(IESBrightnessScale, float);
-}
-
-bool UInterchangeDatasmithAreaLightNode::SetCustomIESBrightnessScale(const float& AttributeValue, bool bAddApplyDelegate)
-{
-	IMPLEMENT_NODE_ATTRIBUTE_SETTER_NODELEGATE(IESBrightnessScale, float);
-}
-
-bool UInterchangeDatasmithAreaLightNode::GetCustomRotation(FRotator& AttributeValue) const
-{
-	IMPLEMENT_NODE_ATTRIBUTE_GETTER(Rotation, FRotator);
-}
-
-bool UInterchangeDatasmithAreaLightNode::SetCustomRotation(const FRotator& AttributeValue, bool bAddApplyDelegate)
-{
-	IMPLEMENT_NODE_ATTRIBUTE_SETTER_NODELEGATE(Rotation, FRotator);
-}
-
 bool UInterchangeDatasmithAreaLightNode::GetCustomSourceRadius(float& AttributeValue) const
 {
 	IMPLEMENT_NODE_ATTRIBUTE_GETTER(SourceRadius, float);
diff --git a/Engine/Plugins/Experimental/Enterprise/DatasmithInterchange/Source/DatasmithInterchange/Private/InterchangeDatasmithMaterialPipeline.cpp b/Engine/Plugins/Experimental/Enterprise/DatasmithInterchange/Source/DatasmithInterchange/Private/InterchangeDatasmithMaterialPipeline.cpp
--- a/Engine/Plugins/Experimental/Enterprise/DatasmithInterchange/Source/DatasmithInterchange/Private/InterchangeDatasmithMaterialPipeline.cpp
+++ b/Engine/Plugins/Experimental/Enterprise/DatasmithInterchange/Source/DatasmithInterchange/Private/InterchangeDatasmithMaterialPipeline.cpp
@@ -329,7 +329,7 @@ void UInterchangeDatasmithMaterialPipeline::UpdateMaterialFactoryNodes(const TAr
 	{
 		const FString FactoryNodeUid = UInterchangeMaterialFactoryNode::GetMaterialFactoryNodeUidFromMaterialNodeUid(ShaderNode->GetUniqueID());
 		UInterchangeFactoryBaseNode* FactoryNode = BaseNodeContainer->GetFactoryNode(FactoryNodeUid);
-		if (!ensure(FactoryNode))
+		if (!ensure(FactoryNode || bIdentifyDuplicateMaterials))
 		{
 			continue;
 		}
diff --git a/Engine/Plugins/Experimental/Enterprise/DatasmithInterchange/Source/DatasmithInterchange/Private/InterchangeDatasmithPipeline.cpp b/Engine/Plugins/Experimental/Enterprise/DatasmithInterchange/Source/DatasmithInterchange/Private/InterchangeDatasmithPipeline.cpp
--- a/Engine/Plugins/Experimental/Enterprise/DatasmithInterchange/Source/DatasmithInterchange/Private/InterchangeDatasmithPipeline.cpp
+++ b/Engine/Plugins/Experimental/Enterprise/DatasmithInterchange/Source/DatasmithInterchange/Private/InterchangeDatasmithPipeline.cpp
@@ -8,7 +8,6 @@
 #include "InterchangeDatasmithLog.h"
 #include "InterchangeDatasmithMaterialPipeline.h"
 #include "InterchangeDatasmithStaticMeshPipeline.h"
-//#include "InterchangeDatasmithTexturePipeline.h"
 #include "InterchangeDatasmithTextureData.h"
 #include "InterchangeDatasmithUtils.h"
 
@@ -20,9 +19,11 @@
 #include "InterchangeSceneImportAssetFactoryNode.h"
 #include "InterchangeSceneVariantSetsFactoryNode.h"
 #include "InterchangeStaticMeshFactoryNode.h"
+#include "InterchangeTexture2DFactoryNode.h"
 #include "InterchangeTexture2DNode.h"
 #include "InterchangeTextureFactoryNode.h"
-#include "InterchangeTexture2DFactoryNode.h"
+#include "InterchangeTextureLightProfileFactoryNode.h"
+#include "InterchangeTextureLightProfileNode.h"
 #include "InterchangeMeshActorFactoryNode.h"
 
 #include "ExternalSource.h"
@@ -119,6 +120,36 @@ UInterchangeDatasmithPipeline::UInterchangeDatasmithPipeline()
 	AnimationPipeline->CommonSkeletalMeshesAndAnimationsProperties = CommonSkeletalMeshesAndAnimationsProperties;
 }
 
+void UInterchangeDatasmithPipeline::AdjustSettingsForContext(EInterchangePipelineContext ImportType, TObjectPtr<UObject> ReimportAsset)
+{
+	Super::AdjustSettingsForContext(ImportType, ReimportAsset);
+
+	if (MaterialPipeline)
+	{
+		MaterialPipeline->AdjustSettingsForContext(ImportType, ReimportAsset);
+	}
+
+	if (MeshPipeline)
+	{
+		MeshPipeline->AdjustSettingsForContext(ImportType, ReimportAsset);
+	}
+
+	if (LevelPipeline)
+	{
+		LevelPipeline->AdjustSettingsForContext(ImportType, ReimportAsset);
+	}
+
+	if (AnimationPipeline)
+	{
+		AnimationPipeline->AdjustSettingsForContext(ImportType, ReimportAsset);
+	}
+}
+
+void UInterchangeDatasmithPipeline::PostDuplicate(bool bDuplicateForPIE)
+{
+	AdjustSettingsForContext(CachePipelineContext, CacheReimportObject.Get());
+}
+
 void UInterchangeDatasmithPipeline::ExecutePipeline(UInterchangeBaseNodeContainer* InBaseNodeContainer, const TArray<UInterchangeSourceData*>& SourceDatas, const FString& ContentBasePath)
 {
 	using namespace UE::DatasmithInterchange;
@@ -144,66 +175,70 @@ void UInterchangeDatasmithPipeline::ExecutePipeline(UInterchangeBaseNodeContaine
 	}
 
 	ExecutePreImportPipelineFunc(MeshPipeline);
-	ExecutePreImportPipelineFunc(LevelPipeline);
-	ExecutePreImportPipelineFunc(AnimationPipeline);
 
-	const FString PackageSubPath = FPaths::GetBaseFilename(SourceDatas[0]->GetFilename());
+	if (CachePipelineContext == EInterchangePipelineContext::SceneImport || CachePipelineContext == EInterchangePipelineContext::SceneReimport)
+	{
+		ExecutePreImportPipelineFunc(LevelPipeline);
+		ExecutePreImportPipelineFunc(AnimationPipeline);
+
+		const FString PackageSubPath = FPaths::GetBaseFilename(SourceDatas[0]->GetFilename());
 
 #if WITH_EDITORONLY_DATA
-	LevelPipeline->SceneImportFactoryNode->SetCustomSubPath(PackageSubPath);
+		LevelPipeline->SceneImportFactoryNode->SetCustomSubPath(PackageSubPath);
 #endif
 
-	// Textures
-	for (UInterchangeTextureFactoryNode* TextureFactoryNode : NodeUtils::GetNodes<UInterchangeTextureFactoryNode>(BaseNodeContainer))
-	{
-		TextureFactoryNode->SetCustomSubPath(FPaths::Combine(PackageSubPath, "Textures"));
-		TextureFactoryNode->SetEnabled(true);
-	}
+		// Textures
+		for (UInterchangeTextureFactoryNode* TextureFactoryNode : NodeUtils::GetNodes<UInterchangeTextureFactoryNode>(BaseNodeContainer))
+		{
+			TextureFactoryNode->SetCustomSubPath(FPaths::Combine(PackageSubPath, "Textures"));
+			TextureFactoryNode->SetEnabled(true);
+		}
 
-	// Materials
-	for (UInterchangeBaseMaterialFactoryNode* MaterialFactoryNode : NodeUtils::GetNodes<UInterchangeBaseMaterialFactoryNode>(BaseNodeContainer))
-	{
-		if (bCreateMaterialReferencesFolders)
+		// Materials
+		for (UInterchangeBaseMaterialFactoryNode* MaterialFactoryNode : NodeUtils::GetNodes<UInterchangeBaseMaterialFactoryNode>(BaseNodeContainer))
 		{
-			if (MaterialFactoryNode->IsA<UInterchangeMaterialFactoryNode>())
+			if (bCreateMaterialReferencesFolders)
 			{
-				MaterialFactoryNode->SetCustomSubPath(FPaths::Combine(PackageSubPath, "Materials/References"));
-			}
-			else if(MaterialFactoryNode->IsA<UInterchangeMaterialFunctionFactoryNode>())
-			{
-				MaterialFactoryNode->SetCustomSubPath(FPaths::Combine(PackageSubPath, "Materials/References/Functions"));
+				if (MaterialFactoryNode->IsA<UInterchangeMaterialFactoryNode>())
+				{
+					MaterialFactoryNode->SetCustomSubPath(FPaths::Combine(PackageSubPath, "Materials/References"));
+				}
+				else if (MaterialFactoryNode->IsA<UInterchangeMaterialFunctionFactoryNode>())
+				{
+					MaterialFactoryNode->SetCustomSubPath(FPaths::Combine(PackageSubPath, "Materials/References/Functions"));
+				}
+				else
+				{
+					MaterialFactoryNode->SetCustomSubPath(FPaths::Combine(PackageSubPath, "Materials"));
+				}
 			}
 			else
 			{
 				MaterialFactoryNode->SetCustomSubPath(FPaths::Combine(PackageSubPath, "Materials"));
 			}
+			MaterialFactoryNode->SetEnabled(true);
 		}
-		else
+
+		// StaticMeshes
+		for (UInterchangeStaticMeshFactoryNode* StaticMeshFactoryNode : NodeUtils::GetNodes<UInterchangeStaticMeshFactoryNode>(BaseNodeContainer))
 		{
-			MaterialFactoryNode->SetCustomSubPath(FPaths::Combine(PackageSubPath, "Materials"));
+			StaticMeshFactoryNode->SetCustomSubPath(FPaths::Combine(PackageSubPath, "Geometries"));
+			StaticMeshFactoryNode->SetEnabled(true);
 		}
-		MaterialFactoryNode->SetEnabled(true);
-	}
 
-	// StaticMeshes
-	for (UInterchangeStaticMeshFactoryNode* StaticMeshFactoryNode : NodeUtils::GetNodes<UInterchangeStaticMeshFactoryNode>(BaseNodeContainer))
-	{
-		StaticMeshFactoryNode->SetCustomSubPath(FPaths::Combine(PackageSubPath, "Geometries"));
-		StaticMeshFactoryNode->SetEnabled(true);
-	}
-
-	// LevelSequences
-	for (UInterchangeLevelSequenceFactoryNode* AnimationTrackSetFactoryNode : NodeUtils::GetNodes<UInterchangeLevelSequenceFactoryNode>(BaseNodeContainer))
-	{
-		AnimationTrackSetFactoryNode->SetCustomSubPath(FPaths::Combine(PackageSubPath, "Animations"));
-		AnimationTrackSetFactoryNode->SetEnabled(true);
-	}
+		// LevelSequences
+		for (UInterchangeLevelSequenceFactoryNode* AnimationTrackSetFactoryNode : NodeUtils::GetNodes<UInterchangeLevelSequenceFactoryNode>(BaseNodeContainer))
+		{
+			AnimationTrackSetFactoryNode->SetCustomSubPath(FPaths::Combine(PackageSubPath, "Animations"));
+			AnimationTrackSetFactoryNode->SetEnabled(true);
+		}
 
-	// LevelVariantSets
-	for (UInterchangeSceneVariantSetsFactoryNode* LevelVariantSetFactoryNode : NodeUtils::GetNodes<UInterchangeSceneVariantSetsFactoryNode>(BaseNodeContainer))
-	{
-		LevelVariantSetFactoryNode->SetCustomSubPath(FPaths::Combine(PackageSubPath, "Variants"));
-		LevelVariantSetFactoryNode->SetEnabled(true);
+		// LevelVariantSets
+		for (UInterchangeSceneVariantSetsFactoryNode* LevelVariantSetFactoryNode : NodeUtils::GetNodes<UInterchangeSceneVariantSetsFactoryNode>(BaseNodeContainer))
+		{
+			LevelVariantSetFactoryNode->SetCustomSubPath(FPaths::Combine(PackageSubPath, "Variants"));
+			LevelVariantSetFactoryNode->SetEnabled(true);
+		}
 	}
 }
 
@@ -221,14 +256,18 @@ void UInterchangeDatasmithPipeline::ExecutePostImportPipeline(const UInterchange
 		MeshPipeline->ScriptedExecutePostImportPipeline(InBaseNodeContainer, NodeKey, CreatedAsset, bIsAReimport);
 	}
 
-	if (LevelPipeline)
+	const bool bSceneImport = CachePipelineContext == EInterchangePipelineContext::SceneImport || CachePipelineContext == EInterchangePipelineContext::SceneReimport;
+	if (bSceneImport)
 	{
-		LevelPipeline->ScriptedExecutePostImportPipeline(InBaseNodeContainer, NodeKey, CreatedAsset, bIsAReimport);
-	}
+		if (LevelPipeline)
+		{
+			LevelPipeline->ScriptedExecutePostImportPipeline(InBaseNodeContainer, NodeKey, CreatedAsset, bIsAReimport);
+		}
 
-	if (AnimationPipeline)
-	{
-		AnimationPipeline->ScriptedExecutePostImportPipeline(InBaseNodeContainer, NodeKey, CreatedAsset, bIsAReimport);
+		if (AnimationPipeline)
+		{
+			AnimationPipeline->ScriptedExecutePostImportPipeline(InBaseNodeContainer, NodeKey, CreatedAsset, bIsAReimport);
+		}
 	}
 }
 
@@ -294,84 +333,84 @@ void UInterchangeDatasmithPipeline::PreImportTextureFactoryNode(UInterchangeText
 			return {};
 		}();
 
-		static_assert(TextureAddress::TA_Wrap == (int)EDatasmithTextureAddress::Wrap && TextureAddress::TA_Mirror == (int)EDatasmithTextureAddress::Mirror, "Texture Address enum doesn't match!");
+	static_assert(TextureAddress::TA_Wrap == (int)EDatasmithTextureAddress::Wrap && TextureAddress::TA_Mirror == (int)EDatasmithTextureAddress::Mirror, "Texture Address enum doesn't match!");
 
-		TOptional< TextureFilter > TexFilter;
-		EDatasmithTextureFilter TextureFilter;
-		if (TextureData.GetCustomTextureFilter(TextureFilter))
+	TOptional< TextureFilter > TexFilter;
+	EDatasmithTextureFilter TextureFilter;
+	if (TextureData.GetCustomTextureFilter(TextureFilter))
+	{
+		switch (TextureFilter)
 		{
-			switch (TextureFilter)
-			{
-			case EDatasmithTextureFilter::Nearest:
-				TexFilter = TextureFilter::TF_Nearest;
-				break;
-			case EDatasmithTextureFilter::Bilinear:
-				TexFilter = TextureFilter::TF_Bilinear;
-				break;
-			case EDatasmithTextureFilter::Trilinear:
-				TexFilter = TextureFilter::TF_Trilinear;
-				break;
-			}
+		case EDatasmithTextureFilter::Nearest:
+			TexFilter = TextureFilter::TF_Nearest;
+			break;
+		case EDatasmithTextureFilter::Bilinear:
+			TexFilter = TextureFilter::TF_Bilinear;
+			break;
+		case EDatasmithTextureFilter::Trilinear:
+			TexFilter = TextureFilter::TF_Trilinear;
+			break;
 		}
+	}
 
-		TOptional< bool > bSrgb;
-		EDatasmithColorSpace ColorSpace;
-		if (TextureData.GetCustomSRGB(ColorSpace))
+	TOptional< bool > bSrgb;
+	EDatasmithColorSpace ColorSpace;
+	if (TextureData.GetCustomSRGB(ColorSpace))
+	{
+		if (ColorSpace == EDatasmithColorSpace::sRGB)
 		{
-			if (ColorSpace == EDatasmithColorSpace::sRGB)
-			{
-				bSrgb = true;
-			}
-			else if (ColorSpace == EDatasmithColorSpace::Linear)
-			{
-				bSrgb = false;
-			}
+			bSrgb = true;
 		}
-
-		EDatasmithTextureAddress AddressX;
-		EDatasmithTextureAddress AddressY;
-		UInterchangeTexture2DFactoryNode* Texture2DFactoryNode = Cast<UInterchangeTexture2DFactoryNode>(TextureFactoryNode);
-		if (Texture2DFactoryNode
-			&& TextureData.GetCustomTextureAddressX(AddressX)
-			&& TextureData.GetCustomTextureAddressY(AddressY))
+		else if (ColorSpace == EDatasmithColorSpace::Linear)
 		{
-			Texture2DFactoryNode->SetCustomAddressX((TextureAddress)AddressX);
-			Texture2DFactoryNode->SetCustomAddressY((TextureAddress)AddressY);
+			bSrgb = false;
 		}
+	}
 
-		if (bSrgb.IsSet())
-		{
-			TextureFactoryNode->SetCustomSRGB(bSrgb.GetValue());
-		}
+	EDatasmithTextureAddress AddressX;
+	EDatasmithTextureAddress AddressY;
+	UInterchangeTexture2DFactoryNode* Texture2DFactoryNode = Cast<UInterchangeTexture2DFactoryNode>(TextureFactoryNode);
+	if (Texture2DFactoryNode
+		&& TextureData.GetCustomTextureAddressX(AddressX)
+		&& TextureData.GetCustomTextureAddressY(AddressY))
+	{
+		Texture2DFactoryNode->SetCustomAddressX((TextureAddress)AddressX);
+		Texture2DFactoryNode->SetCustomAddressY((TextureAddress)AddressY);
+	}
 
-		if (bLocalFlipNormalMapGreenChannel.IsSet())
-		{
-			TextureFactoryNode->SetCustombFlipGreenChannel(bLocalFlipNormalMapGreenChannel.GetValue());
-		}
+	if (bSrgb.IsSet())
+	{
+		TextureFactoryNode->SetCustomSRGB(bSrgb.GetValue());
+	}
 
-		if (MipGenSettings.IsSet())
-		{
-			TextureFactoryNode->SetCustomMipGenSettings(MipGenSettings.GetValue());
-		}
+	if (bLocalFlipNormalMapGreenChannel.IsSet())
+	{
+		TextureFactoryNode->SetCustombFlipGreenChannel(bLocalFlipNormalMapGreenChannel.GetValue());
+	}
 
-		if (LODGroup.IsSet())
-		{
-			TextureFactoryNode->SetCustomLODGroup(LODGroup.GetValue());
-		}
+	if (MipGenSettings.IsSet())
+	{
+		TextureFactoryNode->SetCustomMipGenSettings(MipGenSettings.GetValue());
+	}
 
-		if (CompressionSettings.IsSet())
-		{
-			TextureFactoryNode->SetCustomCompressionSettings(CompressionSettings.GetValue());
-		}
+	if (LODGroup.IsSet())
+	{
+		TextureFactoryNode->SetCustomLODGroup(LODGroup.GetValue());
+	}
 
-		if (RGBCurve.IsSet())
-		{
-			TextureFactoryNode->SetCustomAdjustRGBCurve(RGBCurve.GetValue());
-		}
+	if (CompressionSettings.IsSet())
+	{
+		TextureFactoryNode->SetCustomCompressionSettings(CompressionSettings.GetValue());
+	}
 
-		if (TexFilter.IsSet())
-		{
-			TextureFactoryNode->SetCustomFilter(TexFilter.GetValue());
-		}
+	if (RGBCurve.IsSet())
+	{
+		TextureFactoryNode->SetCustomAdjustRGBCurve(RGBCurve.GetValue());
+	}
+
+	if (TexFilter.IsSet())
+	{
+		TextureFactoryNode->SetCustomFilter(TexFilter.GetValue());
+	}
 }
 #undef LOCTEXT_NAMESPACE
\ No newline at end of file
diff --git a/Engine/Plugins/Experimental/Enterprise/DatasmithInterchange/Source/DatasmithInterchange/Private/InterchangeDatasmithTextureData.cpp b/Engine/Plugins/Experimental/Enterprise/DatasmithInterchange/Source/DatasmithInterchange/Private/InterchangeDatasmithTextureData.cpp
--- a/Engine/Plugins/Experimental/Enterprise/DatasmithInterchange/Source/DatasmithInterchange/Private/InterchangeDatasmithTextureData.cpp
+++ b/Engine/Plugins/Experimental/Enterprise/DatasmithInterchange/Source/DatasmithInterchange/Private/InterchangeDatasmithTextureData.cpp
@@ -11,15 +11,15 @@
 
 namespace UE::DatasmithInterchange
 {
-	const FString FInterchangeDatasmithTextureDataConst::DatasmithTextureDataKey = TEXT("DatasmithTextureData");
-	const FString FInterchangeDatasmithTextureDataConst::FileKey = TEXT("File");
-	const FString FInterchangeDatasmithTextureDataConst::TextureModeKey = TEXT("TextureMode");
-	const FString FInterchangeDatasmithTextureDataConst::TextureFilterKey = TEXT("TextureFilter");
-	const FString FInterchangeDatasmithTextureDataConst::TextureAddressXKey = TEXT("TextureAddressX");
-	const FString FInterchangeDatasmithTextureDataConst::TextureAddressYKey = TEXT("TextureAddressY");
-	const FString FInterchangeDatasmithTextureDataConst::AllowResizeKey = TEXT("AllowResize");
-	const FString FInterchangeDatasmithTextureDataConst::RGBCurveKey = TEXT("RGBCurve");
-	const FString FInterchangeDatasmithTextureDataConst::SRGBKey = TEXT("SRGB");
+	const FString FInterchangeDatasmithTextureDataConst::DatasmithTextureDataKey = TEXT("DatasmithTextureData::DatasmithTextureData");
+	const FString FInterchangeDatasmithTextureDataConst::FileKey = TEXT("DatasmithTextureData::File");
+	const FString FInterchangeDatasmithTextureDataConst::TextureModeKey = TEXT("DatasmithTextureData::TextureMode");
+	const FString FInterchangeDatasmithTextureDataConst::TextureFilterKey = TEXT("DatasmithTextureData::TextureFilter");
+	const FString FInterchangeDatasmithTextureDataConst::TextureAddressXKey = TEXT("DatasmithTextureData::TextureAddressX");
+	const FString FInterchangeDatasmithTextureDataConst::TextureAddressYKey = TEXT("DatasmithTextureData::TextureAddressY");
+	const FString FInterchangeDatasmithTextureDataConst::AllowResizeKey = TEXT("DatasmithTextureData::AllowResize");
+	const FString FInterchangeDatasmithTextureDataConst::RGBCurveKey = TEXT("DatasmithTextureData::RGBCurve");
+	const FString FInterchangeDatasmithTextureDataConst::SRGBKey = TEXT("DatasmithTextureData::SRGB");
 
 	FInterchangeDatasmithTextureDataConst::FInterchangeDatasmithTextureDataConst(const UInterchangeBaseNode* InNode)
 		: ConstBaseNode(InNode)
diff --git a/Engine/Plugins/Experimental/Enterprise/DatasmithInterchange/Source/DatasmithInterchange/Private/InterchangeDatasmithTranslator.cpp b/Engine/Plugins/Experimental/Enterprise/DatasmithInterchange/Source/DatasmithInterchange/Private/InterchangeDatasmithTranslator.cpp
--- a/Engine/Plugins/Experimental/Enterprise/DatasmithInterchange/Source/DatasmithInterchange/Private/InterchangeDatasmithTranslator.cpp
+++ b/Engine/Plugins/Experimental/Enterprise/DatasmithInterchange/Source/DatasmithInterchange/Private/InterchangeDatasmithTranslator.cpp
@@ -31,8 +31,9 @@
 #include "InterchangeShaderGraphNode.h"
 #include "InterchangeSceneNode.h"
 #include "InterchangeTexture2DNode.h"
+#include "InterchangeTextureLightProfileNode.h"
+#include "InterchangeTextureLightProfileFactoryNode.h"
 #include "InterchangeVariantSetNode.h"
-
 #include "StaticMeshOperations.h"
 
 #include "Misc/App.h"
@@ -50,7 +51,7 @@
 
 #define LOCTEXT_NAMESPACE "DatasmithInterchange"
 
-namespace UE::Interchange::Datasmith
+namespace UE::DatasmithInterchange
 {
 #if WITH_EDITOR
 	bool DisplayOptionsDialog(IDatasmithTranslator& Translator)
@@ -182,7 +183,7 @@ bool UInterchangeDatasmithTranslator::Translate(UInterchangeBaseNodeContainer& B
 			TSharedRef<TPromise<bool>, ESPMode::ThreadSafe> Promise = MakeShareable(new TPromise<bool>());
 			TFunction<void()> PromiseKeeper = [&DatasmithTranslator, &Promise]() -> void
 			{
-				bool bShouldImport = UE::Interchange::Datasmith::DisplayOptionsDialog(*DatasmithTranslator);
+				bool bShouldImport = UE::DatasmithInterchange::DisplayOptionsDialog(*DatasmithTranslator);
 				Promise->SetValue(bShouldImport);
 			};
 
@@ -196,7 +197,7 @@ bool UInterchangeDatasmithTranslator::Translate(UInterchangeBaseNodeContainer& B
 		}
 		else
 		{
-			bShouldImport = UE::Interchange::Datasmith::DisplayOptionsDialog(*DatasmithTranslator);
+			bShouldImport = UE::DatasmithInterchange::DisplayOptionsDialog(*DatasmithTranslator);
 		}
 
 		if (!bShouldImport)
@@ -224,15 +225,27 @@ bool UInterchangeDatasmithTranslator::Translate(UInterchangeBaseNodeContainer& B
 		{
 			if (TSharedPtr<IDatasmithTextureElement> TextureElement = DatasmithScene->GetTexture(TextureIndex))
 			{
-				UInterchangeTexture2DNode* TextureNode = NewObject<UInterchangeTexture2DNode>(&BaseNodeContainer);
+				const bool bIsIesProfile = FPaths::GetExtension(TextureElement->GetFile()).Equals(TEXT("ies"), ESearchCase::IgnoreCase);
+				UClass* TextureClass = bIsIesProfile ? UInterchangeTextureLightProfileNode::StaticClass() : UInterchangeTexture2DNode::StaticClass();
+
+				UInterchangeTextureNode* TextureNode = NewObject<UInterchangeTextureNode>(&BaseNodeContainer, TextureClass);
+
 				const FString TextureNodeUid = NodeUtils::TexturePrefix + TextureElement->GetName();
 				const FString DisplayLabel = TextureNameProvider.GenerateUniqueName(TextureElement->GetLabel());
 
 				TextureNode->InitializeNode(TextureNodeUid, DisplayLabel, EInterchangeNodeContainerType::TranslatedAsset);
-				TextureNode->SetPayLoadKey(LexToString(TextureIndex));
-				TextureUtils::ApplyTextureElementToNode(TextureElement.ToSharedRef(), TextureNode);
 
 				BaseNodeContainer.AddNode(TextureNode);
+
+				if (bIsIesProfile)
+				{
+					TextureNode->SetPayLoadKey(TextureElement->GetFile());
+				}
+				else
+				{
+					TextureUtils::ApplyTextureElementToNode(TextureElement.ToSharedRef(), TextureNode);
+					TextureNode->SetPayLoadKey(LexToString(TextureIndex));
+				}
 			}
 		}
 	}
@@ -566,14 +579,6 @@ UInterchangeBaseLightNode* UInterchangeDatasmithTranslator::AddLightNode(UInterc
 			AreaLightNode->SetCustomTemperature(AreaLightElement->GetTemperature());
 		}
 
-		if (AreaLightElement->GetUseIes())
-		{
-			//AreaLightNode->SetCustomIESTexture(const TObjectPtr<class UTextureLightProfile>&AttributeValue, bool bAddApplyDelegate = true);
-			AreaLightNode->SetCustomUseIESBrightness(AreaLightElement->GetUseIesBrightness());
-			AreaLightNode->SetCustomIESBrightnessScale(AreaLightElement->GetIesBrightnessScale());
-			AreaLightNode->SetCustomRotation(AreaLightElement->GetIesRotation().Rotator());
-		}
-
 		AreaLightNode->SetCustomSourceRadius(AreaLightElement->GetSourceRadius());
 		AreaLightNode->SetCustomSourceLength(AreaLightElement->GetSourceLength());
 		AreaLightNode->SetCustomAttenuationRadius(AreaLightElement->GetAttenuationRadius());
@@ -606,6 +611,9 @@ UInterchangeBaseLightNode* UInterchangeDatasmithTranslator::AddLightNode(UInterc
 		ensure(false);
 		LightNode = NewObject<UInterchangeLightNode>(&BaseNodeContainer);
 	}
+
+	ProcessIesProfile(BaseNodeContainer, *LightActor, Cast<UInterchangeLightNode>(LightNode));
+
 	const FString LightUid = NodeUtils::LightPrefix + LightActor->GetName();
 	LightNode->InitializeNode(LightUid, LightActor->GetLabel(), EInterchangeNodeContainerType::TranslatedAsset);
 	BaseNodeContainer.AddNode(LightNode);
@@ -613,6 +621,56 @@ UInterchangeBaseLightNode* UInterchangeDatasmithTranslator::AddLightNode(UInterc
 	return LightNode;
 }
 
+void UInterchangeDatasmithTranslator::ProcessIesProfile(UInterchangeBaseNodeContainer& BaseNodeContainer, const IDatasmithLightActorElement& LightElement, UInterchangeLightNode* LightNode) const
+{
+	using namespace UE::DatasmithImporter;
+	using namespace UE::DatasmithInterchange;
+
+	if (!LightNode || !LightElement.GetUseIes())
+	{
+		return;
+	}
+
+	bool bUpdateLightNode = false;
+
+	FString ProfileNodeUid = NodeUtils::TexturePrefix + LightElement.GetName() + TEXT("_IES");
+	const FString DisplayLabel = FString(LightElement.GetName()) + TEXT("_IES");
+
+	if (FPaths::FileExists(LightElement.GetIesTexturePathName()))
+	{
+		UInterchangeTextureNode* TextureNode = NewObject<UInterchangeTextureLightProfileNode>(&BaseNodeContainer);
+		TextureNode->InitializeNode(ProfileNodeUid, DisplayLabel, EInterchangeNodeContainerType::TranslatedAsset);
+		BaseNodeContainer.AddNode(TextureNode);
+		bUpdateLightNode = true;
+	}
+	else if(FSoftObjectPath(LightElement.GetIesTexturePathName()).IsValid())
+	{
+		FString IESFactoryTextureId = UInterchangeFactoryBaseNode::BuildFactoryNodeUid(ProfileNodeUid);
+		UInterchangeTextureLightProfileFactoryNode* FactoryNode = NewObject<UInterchangeTextureLightProfileFactoryNode>(&BaseNodeContainer);
+		FactoryNode->InitializeNode(IESFactoryTextureId, DisplayLabel, EInterchangeNodeContainerType::FactoryData);
+		FactoryNode->SetCustomReferenceObject(FSoftObjectPath(LightElement.GetIesTexturePathName()));
+		BaseNodeContainer.AddNode(FactoryNode);
+		bUpdateLightNode = true;
+	}
+	else
+	{
+		const FString TextureNodeUid = NodeUtils::TexturePrefix + FDatasmithUtils::SanitizeObjectName(LightElement.GetIesTexturePathName());
+		if (BaseNodeContainer.GetNode(TextureNodeUid))
+		{
+			ProfileNodeUid = TextureNodeUid;
+			bUpdateLightNode = true;
+		}
+	}
+
+	if (bUpdateLightNode)
+	{
+		LightNode->SetCustomIESTexture(ProfileNodeUid);
+		LightNode->SetCustomUseIESBrightness(LightElement.GetUseIesBrightness());
+		LightNode->SetCustomIESBrightnessScale(LightElement.GetIesBrightnessScale());
+		LightNode->SetCustomRotation(LightElement.GetIesRotation().Rotator());
+	}
+}
+
 UInterchangeDecalNode* UInterchangeDatasmithTranslator::AddDecalNode(UInterchangeBaseNodeContainer& BaseNodeContainer, const TSharedRef<IDatasmithDecalActorElement>& DecalActor) const
 {
 	using namespace UE::DatasmithInterchange;
@@ -668,6 +726,8 @@ TOptional<UE::Interchange::FImportImage> UInterchangeDatasmithTranslator::GetTex
 		return TOptional<UE::Interchange::FImportImage>();
 	}
 
+	PayloadSourceData->ClearInternalFlags(EInternalObjectFlags::Async);
+
 	UInterchangeTranslatorBase* SourceTranslator = UInterchangeManager::GetInterchangeManager().GetTranslatorForSourceData(PayloadSourceData);
 	FGCObjectScopeGuard ScopedSourceTranslator(SourceTranslator);
 	const IInterchangeTexturePayloadInterface* TextureTranslator = Cast< IInterchangeTexturePayloadInterface >(SourceTranslator);
@@ -676,6 +736,7 @@ TOptional<UE::Interchange::FImportImage> UInterchangeDatasmithTranslator::GetTex
 		return TOptional<UE::Interchange::FImportImage>();
 	}
 
+	SourceTranslator->ClearInternalFlags(EInternalObjectFlags::Async);
 	SourceTranslator->SetResultsContainer(Results);
 
 	AlternateTexturePath = TextureElement->GetFile();
@@ -683,6 +744,40 @@ TOptional<UE::Interchange::FImportImage> UInterchangeDatasmithTranslator::GetTex
 	return TextureTranslator->GetTexturePayloadData(PayloadKey, AlternateTexturePath);
 }
 
+TOptional<UE::Interchange::FImportLightProfile> UInterchangeDatasmithTranslator::GetLightProfilePayloadData(const FString& PayloadKey, TOptional<FString>& AlternateTexturePath) const
+{
+	if (!LoadedExternalSource || !LoadedExternalSource->GetDatasmithScene())
+	{
+		return TOptional<UE::Interchange::FImportLightProfile>();
+	}
+
+	UInterchangeSourceData* PayloadSourceData = UInterchangeManager::GetInterchangeManager().CreateSourceData(PayloadKey);
+	FGCObjectScopeGuard ScopedSourceData(PayloadSourceData);
+	if (!PayloadSourceData)
+	{
+		return TOptional<UE::Interchange::FImportLightProfile>();
+	}
+
+	PayloadSourceData->ClearInternalFlags(EInternalObjectFlags::Async);
+
+	UInterchangeTranslatorBase* SourceTranslator = UInterchangeManager::GetInterchangeManager().GetTranslatorForSourceData(PayloadSourceData);
+	FGCObjectScopeGuard ScopedSourceTranslator(SourceTranslator);
+	const IInterchangeTextureLightProfilePayloadInterface* TextureTranslator = Cast< IInterchangeTextureLightProfilePayloadInterface >(SourceTranslator);
+	if (!ensure(TextureTranslator))
+	{
+		return TOptional<UE::Interchange::FImportLightProfile>();
+	}
+
+	SourceTranslator->ClearInternalFlags(EInternalObjectFlags::Async);
+	SourceTranslator->SetResultsContainer(Results);
+
+	AlternateTexturePath = PayloadKey;
+
+	AlternateTexturePath = PayloadKey;
+
+	return TextureTranslator->GetLightProfilePayloadData(PayloadKey, AlternateTexturePath);
+}
+
 TFuture<TOptional<UE::Interchange::FMeshPayloadData>> UInterchangeDatasmithTranslator::GetMeshPayloadData(const FInterchangeMeshPayLoadKey& PayLoadKey, const FTransform& MeshGlobalTransform) const
 {
 	TPromise<TOptional<UE::Interchange::FMeshPayloadData>> EmptyPromise;
@@ -740,8 +835,6 @@ TFuture<TOptional<UE::Interchange::FMeshPayloadData>> UInterchangeDatasmithTrans
 
 TFuture<TOptional<UE::Interchange::FAnimationPayloadData>> UInterchangeDatasmithTranslator::GetAnimationPayloadData(const FInterchangeAnimationPayLoadKey& PayLoadKey, const double BakeFrequency, const double RangeStartSecond, const double RangeStopSecond) const
 {
-	UE::Interchange::FAnimationPayloadData TransformPayloadData(PayLoadKey.Type);
-
 	TPromise<TOptional<UE::Interchange::FAnimationPayloadData>> EmptyPromise;
 	EmptyPromise.SetValue(TOptional<UE::Interchange::FAnimationPayloadData>());
 	
@@ -765,46 +858,22 @@ TFuture<TOptional<UE::Interchange::FAnimationPayloadData>> UInterchangeDatasmith
 		FrameRate = PayloadDescPtr->Key;
 	}
 
-	switch (PayLoadKey.Type)
+	if (PayLoadKey.Type != EInterchangeAnimationPayLoadType::NONE)
 	{
-	case EInterchangeAnimationPayLoadType::CURVE:
-	case EInterchangeAnimationPayLoadType::MORPHTARGETCURVE:
-		{
-			return Async(EAsyncExecution::TaskGraph, [this, AnimationElement = MoveTemp(AnimationElement), FrameRate, &TransformPayloadData]
-				{
-					
-					TOptional<UE::Interchange::FAnimationPayloadData> Result;
+		return Async(EAsyncExecution::TaskGraph, [this, AnimationElement = MoveTemp(AnimationElement), FrameRate, PayLoadType = PayLoadKey.Type]
+			{
 
-					if (UE::DatasmithInterchange::AnimUtils::GetAnimationPayloadData(*AnimationElement, FrameRate, TransformPayloadData.Curves))
-					{
-						Result.Emplace(MoveTemp(TransformPayloadData));
-					}
+				UE::Interchange::FAnimationPayloadData TransformPayloadData(PayLoadType);
+				TOptional<UE::Interchange::FAnimationPayloadData> Result;
 
-					return Result;
-					}
-				);
-		}
-		break;
-	case EInterchangeAnimationPayLoadType::STEPCURVE:
-		{
-			return Async(EAsyncExecution::TaskGraph, [this, AnimationElement = MoveTemp(AnimationElement), FrameRate, &TransformPayloadData]
+				if (UE::DatasmithInterchange::AnimUtils::GetAnimationPayloadData(*AnimationElement, FrameRate, PayLoadType, TransformPayloadData))
 				{
-					TOptional<UE::Interchange::FAnimationPayloadData> Result;
-
-					if (UE::DatasmithInterchange::AnimUtils::GetAnimationPayloadData(*AnimationElement, FrameRate, TransformPayloadData.StepCurves))
-					{
-						Result.Emplace(MoveTemp(TransformPayloadData));
-					}
-
-					return Result;
+					Result.Emplace(MoveTemp(TransformPayloadData));
 				}
-			);
-		}
-		break;
-	case EInterchangeAnimationPayLoadType::BAKED:
-	case EInterchangeAnimationPayLoadType::NONE:
-	default:
-		break;
+
+				return Result;
+			}
+		);
 	}
 
 	return EmptyPromise.GetFuture();
diff --git a/Engine/Plugins/Experimental/Enterprise/DatasmithInterchange/Source/DatasmithInterchange/Public/InterchangeDatasmithAreaLightNode.h b/Engine/Plugins/Experimental/Enterprise/DatasmithInterchange/Source/DatasmithInterchange/Public/InterchangeDatasmithAreaLightNode.h
--- a/Engine/Plugins/Experimental/Enterprise/DatasmithInterchange/Source/DatasmithInterchange/Public/InterchangeDatasmithAreaLightNode.h
+++ b/Engine/Plugins/Experimental/Enterprise/DatasmithInterchange/Source/DatasmithInterchange/Public/InterchangeDatasmithAreaLightNode.h
@@ -55,36 +55,6 @@ public:
 	UFUNCTION(BlueprintCallable, Category = "Interchange | Node | Datasmith")
 	bool SetCustomColor(const FLinearColor& AttributeValue, bool bAddApplyDelegate = true);
 
-	/**
-	 * Get the NodeUid of the imported IES Texture.
-	 */
-	//UFUNCTION(BlueprintCallable, Category = "Interchange | Node | Datasmith")
-	bool GetCustomIESTexture(FString& AttributeValue) const;
-
-	/**
-	 * Set the NodeUid of the imported IES Texture.
-	 */
-	//UFUNCTION(BlueprintCallable, Category = "Interchange | Node | Datasmith")
-	bool SetCustomIESTexture(const FString& AttributeValue);
-
-	UFUNCTION(BlueprintCallable, Category = "Interchange | Node | Datasmith")
-	bool GetCustomUseIESBrightness(bool& AttributeValue) const;
-
-	UFUNCTION(BlueprintCallable, Category = "Interchange | Node | Datasmith")
-	bool SetCustomUseIESBrightness(const bool& AttributeValue, bool bAddApplyDelegate = true);
-
-	UFUNCTION(BlueprintCallable, Category = "Interchange | Node | Datasmith")
-	bool GetCustomIESBrightnessScale(float& AttributeValue) const;
-
-	UFUNCTION(BlueprintCallable, Category = "Interchange | Node | Datasmith")
-	bool SetCustomIESBrightnessScale(const float& AttributeValue, bool bAddApplyDelegate = true);
-
-	UFUNCTION(BlueprintCallable, Category = "Interchange | Node | Datasmith")
-	bool GetCustomRotation(FRotator& AttributeValue) const;
-
-	UFUNCTION(BlueprintCallable, Category = "Interchange | Node | Datasmith")
-	bool SetCustomRotation(const FRotator& AttributeValue, bool bAddApplyDelegate = true);
-
 	UFUNCTION(BlueprintCallable, Category = "Interchange | Node | Datasmith")
 	bool GetCustomSourceRadius(float& AttributeValue) const;
 
@@ -114,10 +84,6 @@ private:
 	static  UE::Interchange::FAttributeKey Macro_CustomLightShapeKey;
 	static  UE::Interchange::FAttributeKey Macro_CustomDimensionsKey;
 	static  UE::Interchange::FAttributeKey Macro_CustomColorKey;
-	static  UE::Interchange::FAttributeKey Macro_CustomIESTextureKey;
-	static  UE::Interchange::FAttributeKey Macro_CustomUseIESBrightnessKey;
-	static  UE::Interchange::FAttributeKey Macro_CustomIESBrightnessScaleKey;
-	static  UE::Interchange::FAttributeKey Macro_CustomRotationKey;
 	static  UE::Interchange::FAttributeKey Macro_CustomSourceRadiusKey;
 	static  UE::Interchange::FAttributeKey Macro_CustomSourceLengthKey;
 	static  UE::Interchange::FAttributeKey Macro_CustomSpotlightInnerAngleKey;
diff --git a/Engine/Plugins/Experimental/Enterprise/DatasmithInterchange/Source/DatasmithInterchange/Public/InterchangeDatasmithPipeline.h b/Engine/Plugins/Experimental/Enterprise/DatasmithInterchange/Source/DatasmithInterchange/Public/InterchangeDatasmithPipeline.h
--- a/Engine/Plugins/Experimental/Enterprise/DatasmithInterchange/Source/DatasmithInterchange/Public/InterchangeDatasmithPipeline.h
+++ b/Engine/Plugins/Experimental/Enterprise/DatasmithInterchange/Source/DatasmithInterchange/Public/InterchangeDatasmithPipeline.h
@@ -56,9 +56,15 @@ class DATASMITHINTERCHANGE_API UInterchangeDatasmithPipeline : public UInterchan
 	UPROPERTY(EditAnywhere, Category = "Materials")
 	bool bCreateMaterialReferencesFolders = true;
 
+public:
+	/** Begin UObject overrides*/
+	virtual void PostDuplicate(bool bDuplicateForPIE) override;
+	/** End UObject overrides */
+
 protected:
 	virtual void ExecutePipeline(UInterchangeBaseNodeContainer* InBaseNodeContainer, const TArray<UInterchangeSourceData*>& SourceDatas, const FString& ContentBasePath) override;
 	virtual void ExecutePostImportPipeline(const UInterchangeBaseNodeContainer* InBaseNodeContainer, const FString& NodeKey, UObject* CreatedAsset, bool bIsAReimport) override;
+	virtual void AdjustSettingsForContext(EInterchangePipelineContext ImportType, TObjectPtr<UObject> ReimportAsset) override;
 
 	virtual bool CanExecuteOnAnyThread(EInterchangePipelineTask PipelineTask) override
 	{
diff --git a/Engine/Plugins/Experimental/Enterprise/DatasmithInterchange/Source/DatasmithInterchange/Public/InterchangeDatasmithTranslator.h b/Engine/Plugins/Experimental/Enterprise/DatasmithInterchange/Source/DatasmithInterchange/Public/InterchangeDatasmithTranslator.h
--- a/Engine/Plugins/Experimental/Enterprise/DatasmithInterchange/Source/DatasmithInterchange/Public/InterchangeDatasmithTranslator.h
+++ b/Engine/Plugins/Experimental/Enterprise/DatasmithInterchange/Source/DatasmithInterchange/Public/InterchangeDatasmithTranslator.h
@@ -5,10 +5,13 @@
 #include "CoreMinimal.h"
 
 #include "Animation/InterchangeAnimationPayloadInterface.h"
+#include "InterchangeCommonAnimationPayload.h"
 #include "InterchangeTranslatorBase.h"
 #include "Mesh/InterchangeMeshPayloadInterface.h"
 #include "Texture/InterchangeTexturePayloadData.h"
 #include "Texture/InterchangeTexturePayloadInterface.h"
+#include "Texture/InterchangeTextureLightProfilePayloadData.h"
+#include "Texture/InterchangeTextureLightProfilePayloadInterface.h"
 #include "Scene/InterchangeVariantSetPayloadInterface.h"
 
 #include "Async/Async.h"
@@ -29,6 +32,7 @@ class IDatasmithTransformAnimationElement;
 class UInterchangePhysicalCameraNode;
 class UInterchangeBaseLightNode;
 class UInterchangeDecalNode;
+class UInterchangeLightNode;
 class UInterchangeSceneNode;
 
 namespace UE::Interchange
@@ -47,13 +51,13 @@ namespace UE::DatasmithImporter
 namespace UE::DatasmithInterchange::AnimUtils
 {
 	typedef TPair<float, TSharedPtr<IDatasmithBaseAnimationElement>> FAnimationPayloadDesc;
-	extern bool GetAnimationPayloadData(const IDatasmithBaseAnimationElement& AnimationElement, float FrameRate, TArray<FRichCurve>& Curves);
-	extern bool GetAnimationPayloadData(const IDatasmithBaseAnimationElement& AnimationElement, float FrameRate, TArray<FInterchangeStepCurve>& StepCurves);
+	extern bool GetAnimationPayloadData(const IDatasmithBaseAnimationElement& AnimationElement, float FrameRate, EInterchangeAnimationPayLoadType PayLoadType, UE::Interchange::FAnimationPayloadData& PayLoadData);
 }
 
 UCLASS(BlueprintType, Experimental)
 class DATASMITHINTERCHANGE_API UInterchangeDatasmithTranslator : public UInterchangeTranslatorBase
 	, public IInterchangeTexturePayloadInterface
+	, public IInterchangeTextureLightProfilePayloadInterface
 	, public IInterchangeMeshPayloadInterface
 	, public IInterchangeAnimationPayloadInterface
 	, public IInterchangeVariantSetPayloadInterface
@@ -87,6 +91,10 @@ public:
 	virtual TOptional<UE::Interchange::FImportImage> GetTexturePayloadData(const FString& PayloadKey, TOptional<FString>& AlternateTexturePath) const override;
 	/* IInterchangeTexturePayloadInterface End */
 
+	/* IInterchangeTextureLightProfilePayloadInterface Begin */
+	virtual TOptional<UE::Interchange::FImportLightProfile> GetLightProfilePayloadData(const FString& PayloadKey, TOptional<FString>& AlternateTexturePath) const override;
+	/* IInterchangeTextureLightProfilePayloadInterface End */
+
 	/* IInterchangeStaticMeshPayloadInterface Begin */
 	virtual TFuture<TOptional<UE::Interchange::FMeshPayloadData>> GetMeshPayloadData(const FInterchangeMeshPayLoadKey& PayLoadKey, const FTransform& MeshGlobalTransform) const override;
 	/* IInterchangeStaticMeshPayloadInterface End */
@@ -109,6 +117,8 @@ private:
 
 	UInterchangeDecalNode* AddDecalNode(UInterchangeBaseNodeContainer& BaseNodeContainer, const TSharedRef<IDatasmithDecalActorElement>& DecalActor) const;
 
+	void ProcessIesProfile(UInterchangeBaseNodeContainer& BaseNodeContainer, const IDatasmithLightActorElement& LightElement, UInterchangeLightNode* LightNode) const;
+
 	mutable TSharedPtr<UE::DatasmithImporter::FExternalSource> LoadedExternalSource;
 
 	mutable uint64 StartTime = 0;
diff --git a/Engine/Plugins/Interchange/Runtime/Source/FactoryNodes/Private/InterchangeLightFactoryNode.cpp b/Engine/Plugins/Interchange/Runtime/Source/FactoryNodes/Private/InterchangeLightFactoryNode.cpp
--- a/Engine/Plugins/Interchange/Runtime/Source/FactoryNodes/Private/InterchangeLightFactoryNode.cpp
+++ b/Engine/Plugins/Interchange/Runtime/Source/FactoryNodes/Private/InterchangeLightFactoryNode.cpp
@@ -88,6 +88,36 @@ bool UInterchangeLightFactoryNode::SetCustomIESTexture(const FString& AttributeV
 	IMPLEMENT_NODE_ATTRIBUTE_SETTER_NODELEGATE(IESTexture, FString)
 }
 
+bool UInterchangeLightFactoryNode::GetCustomUseIESBrightness(bool& AttributeValue) const
+{
+	IMPLEMENT_NODE_ATTRIBUTE_GETTER(UseIESBrightness, bool);
+}
+
+bool UInterchangeLightFactoryNode::SetCustomUseIESBrightness(const bool& AttributeValue, bool bAddApplyDelegate)
+{
+	IMPLEMENT_NODE_ATTRIBUTE_SETTER_NODELEGATE(UseIESBrightness, bool);
+}
+
+bool UInterchangeLightFactoryNode::GetCustomIESBrightnessScale(float& AttributeValue) const
+{
+	IMPLEMENT_NODE_ATTRIBUTE_GETTER(IESBrightnessScale, float);
+}
+
+bool UInterchangeLightFactoryNode::SetCustomIESBrightnessScale(const float& AttributeValue, bool bAddApplyDelegate)
+{
+	IMPLEMENT_NODE_ATTRIBUTE_SETTER_NODELEGATE(IESBrightnessScale, float);
+}
+
+bool UInterchangeLightFactoryNode::GetCustomRotation(FRotator& AttributeValue) const
+{
+	IMPLEMENT_NODE_ATTRIBUTE_GETTER(Rotation, FRotator);
+}
+
+bool UInterchangeLightFactoryNode::SetCustomRotation(const FRotator& AttributeValue, bool bAddApplyDelegate)
+{
+	IMPLEMENT_NODE_ATTRIBUTE_SETTER_NODELEGATE(Rotation, FRotator);
+}
+
 void UInterchangeLightFactoryNode::CopyWithObject(const UInterchangeFactoryBaseNode* SourceNode, UObject* Object)
 {
 	Super::CopyWithObject(SourceNode, Object);
diff --git a/Engine/Plugins/Interchange/Runtime/Source/FactoryNodes/Public/InterchangeLightFactoryNode.h b/Engine/Plugins/Interchange/Runtime/Source/FactoryNodes/Public/InterchangeLightFactoryNode.h
--- a/Engine/Plugins/Interchange/Runtime/Source/FactoryNodes/Public/InterchangeLightFactoryNode.h
+++ b/Engine/Plugins/Interchange/Runtime/Source/FactoryNodes/Public/InterchangeLightFactoryNode.h
@@ -79,6 +79,24 @@ public:
 	UFUNCTION(BlueprintCallable, Category = "Interchange | Node | LightFactory")
 	bool SetCustomIESTexture(const FString& AttributeValue);
 
+	UFUNCTION(BlueprintCallable, Category = "Interchange | Node | LightFactory")
+	bool GetCustomUseIESBrightness(bool& AttributeValue) const;
+
+	UFUNCTION(BlueprintCallable, Category = "Interchange | Node | LightFactory")
+	bool SetCustomUseIESBrightness(const bool& AttributeValue, bool bAddApplyDelegate = true);
+
+	UFUNCTION(BlueprintCallable, Category = "Interchange | Node | LightFactory")
+	bool GetCustomIESBrightnessScale(float& AttributeValue) const;
+
+	UFUNCTION(BlueprintCallable, Category = "Interchange | Node | LightFactory")
+	bool SetCustomIESBrightnessScale(const float& AttributeValue, bool bAddApplyDelegate = true);
+
+	UFUNCTION(BlueprintCallable, Category = "Interchange | Node | LightFactory")
+	bool GetCustomRotation(FRotator& AttributeValue) const;
+
+	UFUNCTION(BlueprintCallable, Category = "Interchange | Node | LightFactory")
+	bool SetCustomRotation(const FRotator& AttributeValue, bool bAddApplyDelegate = true);
+
 	virtual void CopyWithObject(const UInterchangeFactoryBaseNode* SourceNode, UObject* Object) override;
 
 private:
@@ -86,6 +104,9 @@ private:
 	IMPLEMENT_NODE_ATTRIBUTE_KEY(IntensityUnits)
 	IMPLEMENT_NODE_ATTRIBUTE_KEY(AttenuationRadius)
 	IMPLEMENT_NODE_ATTRIBUTE_KEY(IESTexture)
+	IMPLEMENT_NODE_ATTRIBUTE_KEY(UseIESBrightness)
+	IMPLEMENT_NODE_ATTRIBUTE_KEY(IESBrightnessScale)
+	IMPLEMENT_NODE_ATTRIBUTE_KEY(Rotation)
 };
 
 UCLASS(BlueprintType)
diff --git a/Engine/Plugins/Interchange/Runtime/Source/Import/Private/Scene/InterchangeLightActorFactory.cpp b/Engine/Plugins/Interchange/Runtime/Source/Import/Private/Scene/InterchangeLightActorFactory.cpp
--- a/Engine/Plugins/Interchange/Runtime/Source/Import/Private/Scene/InterchangeLightActorFactory.cpp
+++ b/Engine/Plugins/Interchange/Runtime/Source/Import/Private/Scene/InterchangeLightActorFactory.cpp
@@ -30,14 +30,31 @@ UObject* UInterchangeLightActorFactory::ProcessActor(AActor& SpawnedActor, const
 			{
 				if (FString IESTexture; LightFactoryNode->GetCustomIESTexture(IESTexture))
 				{
-					IESTexture = UInterchangeFactoryBaseNode::BuildFactoryNodeUid(IESTexture);
-					if (const UInterchangeTextureLightProfileFactoryNode* TextureFactoryNode = Cast<UInterchangeTextureLightProfileFactoryNode>(NodeContainer.GetNode(IESTexture)))
+					FString IESFactoryTextureId = UInterchangeFactoryBaseNode::BuildFactoryNodeUid(IESTexture);
+					if (const UInterchangeTextureLightProfileFactoryNode* TextureFactoryNode = Cast<UInterchangeTextureLightProfileFactoryNode>(NodeContainer.GetNode(IESFactoryTextureId)))
 					{
 						FSoftObjectPath ReferenceObject;
 						TextureFactoryNode->GetCustomReferenceObject(ReferenceObject);
 						if (UTextureLightProfile* Texture = Cast<UTextureLightProfile>(ReferenceObject.TryLoad()))
 						{
 							LightComponent->SetIESTexture(Texture);
+							
+							if (bool bUseIESBrightness; LightFactoryNode->GetCustomUseIESBrightness(bUseIESBrightness))
+							{
+								LightComponent->SetUseIESBrightness(bUseIESBrightness);
+							}
+							
+							if (float IESBrightnessScale; LightFactoryNode->GetCustomIESBrightnessScale(IESBrightnessScale))
+							{
+								LightComponent->SetIESBrightnessScale(IESBrightnessScale);
+							}
+							
+							if (FRotator Rotation; LightFactoryNode->GetCustomRotation(Rotation))
+							{
+								// Apply Ies rotation to light orientation
+								FQuat NewLightRotation = LightComponent->GetRelativeRotation().Quaternion() * Rotation.Quaternion();
+								LightComponent->SetRelativeRotation(NewLightRotation);
+							}
 						}
 					}
 				}
diff --git a/Engine/Plugins/Interchange/Runtime/Source/Nodes/Private/InterchangeLightNode.cpp b/Engine/Plugins/Interchange/Runtime/Source/Nodes/Private/InterchangeLightNode.cpp
--- a/Engine/Plugins/Interchange/Runtime/Source/Nodes/Private/InterchangeLightNode.cpp
+++ b/Engine/Plugins/Interchange/Runtime/Source/Nodes/Private/InterchangeLightNode.cpp
@@ -87,6 +87,37 @@ bool UInterchangeLightNode::SetCustomIESTexture(const FString& AttributeValue)
 	IMPLEMENT_NODE_ATTRIBUTE_SETTER_NODELEGATE(IESTexture, FString)
 }
 
+bool UInterchangeLightNode::GetCustomUseIESBrightness(bool& AttributeValue) const
+{
+	IMPLEMENT_NODE_ATTRIBUTE_GETTER(UseIESBrightness, bool);
+}
+
+bool UInterchangeLightNode::SetCustomUseIESBrightness(const bool& AttributeValue, bool bAddApplyDelegate)
+{
+	IMPLEMENT_NODE_ATTRIBUTE_SETTER_NODELEGATE(UseIESBrightness, bool);
+}
+
+bool UInterchangeLightNode::GetCustomIESBrightnessScale(float& AttributeValue) const
+{
+	IMPLEMENT_NODE_ATTRIBUTE_GETTER(IESBrightnessScale, float);
+}
+
+bool UInterchangeLightNode::SetCustomIESBrightnessScale(const float& AttributeValue, bool bAddApplyDelegate)
+{
+	IMPLEMENT_NODE_ATTRIBUTE_SETTER_NODELEGATE(IESBrightnessScale, float);
+}
+
+bool UInterchangeLightNode::GetCustomRotation(FRotator& AttributeValue) const
+{
+	IMPLEMENT_NODE_ATTRIBUTE_GETTER(Rotation, FRotator);
+}
+
+bool UInterchangeLightNode::SetCustomRotation(const FRotator& AttributeValue, bool bAddApplyDelegate)
+{
+	IMPLEMENT_NODE_ATTRIBUTE_SETTER_NODELEGATE(Rotation, FRotator);
+}
+
+
 FString UInterchangePointLightNode::GetTypeName() const
 {
 	return TEXT("PointLightNode");
@@ -165,4 +196,4 @@ bool UInterchangeRectLightNode::SetCustomSourceHeight(float AttributeValue)
 FString UInterchangeDirectionalLightNode::GetTypeName() const
 {
 	return TEXT("DirectionalLightNode");
-}
\ No newline at end of file
+}
diff --git a/Engine/Plugins/Interchange/Runtime/Source/Nodes/Public/InterchangeLightNode.h b/Engine/Plugins/Interchange/Runtime/Source/Nodes/Public/InterchangeLightNode.h
--- a/Engine/Plugins/Interchange/Runtime/Source/Nodes/Public/InterchangeLightNode.h
+++ b/Engine/Plugins/Interchange/Runtime/Source/Nodes/Public/InterchangeLightNode.h
@@ -30,28 +30,28 @@ public:
 	 */
 	virtual FString GetTypeName() const override;
 
-	UFUNCTION(BlueprintCallable, Category = "Interchange | Node | Light")
+	UFUNCTION(BlueprintCallable, Category = "Interchange | Node | BaseLight")
 	bool GetCustomLightColor(FLinearColor& AttributeValue) const;
 
-	UFUNCTION(BlueprintCallable, Category = "Interchange | Node | Light")
+	UFUNCTION(BlueprintCallable, Category = "Interchange | Node | BaseLight")
 	bool SetCustomLightColor(const FLinearColor& AttributeValue);
 
-	UFUNCTION(BlueprintCallable, Category = "Interchange | Node | Light")
+	UFUNCTION(BlueprintCallable, Category = "Interchange | Node | BaseLight")
 	bool GetCustomIntensity(float& AttributeValue) const;
 
-	UFUNCTION(BlueprintCallable, Category = "Interchange | Node | Light")
+	UFUNCTION(BlueprintCallable, Category = "Interchange | Node | BaseLight")
 	bool SetCustomIntensity(float AttributeValue);
 
-	UFUNCTION(BlueprintCallable, Category = "Interchange | Node | Light")
+	UFUNCTION(BlueprintCallable, Category = "Interchange | Node | BaseLight")
 	bool GetCustomTemperature(float& AttributeValue) const;
 
-	UFUNCTION(BlueprintCallable, Category = "Interchange | Node | Light")
+	UFUNCTION(BlueprintCallable, Category = "Interchange | Node | BaseLight")
 	bool SetCustomTemperature(float AttributeValue);
 
-	UFUNCTION(BlueprintCallable, Category = "Interchange | Node | Light")
+	UFUNCTION(BlueprintCallable, Category = "Interchange | Node | BaseLight")
 	bool GetCustomUseTemperature(bool & AttributeValue) const;
 
-	UFUNCTION(BlueprintCallable, Category = "Interchange | Node | Light")
+	UFUNCTION(BlueprintCallable, Category = "Interchange | Node | BaseLight")
 	bool SetCustomUseTemperature(bool AttributeValue);
 
 private:
@@ -71,29 +71,51 @@ public:
 
 	virtual FString GetTypeName() const override;
 
-	UFUNCTION(BlueprintCallable, Category = "Interchange | Node | LocalLight")
+	UFUNCTION(BlueprintCallable, Category = "Interchange | Node | Light")
 	bool GetCustomIntensityUnits(EInterchangeLightUnits& AttributeValue) const;
 
-	UFUNCTION(BlueprintCallable, Category = "Interchange | Node | LocalLight")
+	UFUNCTION(BlueprintCallable, Category = "Interchange | Node | Light")
 	bool SetCustomIntensityUnits(const EInterchangeLightUnits & AttributeValue);
 
-	UFUNCTION(BlueprintCallable, Category = "Interchange | Node | LocalLight")
+	UFUNCTION(BlueprintCallable, Category = "Interchange | Node | Light")
 	bool GetCustomAttenuationRadius(float& AttributeValue) const;
 
-	UFUNCTION(BlueprintCallable, Category = "Interchange | Node | LocalLight")
+	UFUNCTION(BlueprintCallable, Category = "Interchange | Node | Light")
 	bool SetCustomAttenuationRadius(float AttributeValue);
 
-	UFUNCTION(BlueprintCallable, Category = "Interchange | Node | LocalLight")
+	UFUNCTION(BlueprintCallable, Category = "Interchange | Node | Light")
 	bool GetCustomIESTexture(FString& AttributeValue) const;
 
-	UFUNCTION(BlueprintCallable, Category = "Interchange | Node | LocalLight")
+	UFUNCTION(BlueprintCallable, Category = "Interchange | Node | Light")
 	bool SetCustomIESTexture(const FString& AttributeValue);
 
+	UFUNCTION(BlueprintCallable, Category = "Interchange | Node | Light")
+	bool GetCustomUseIESBrightness(bool& AttributeValue) const;
+
+	UFUNCTION(BlueprintCallable, Category = "Interchange | Node | Light")
+	bool SetCustomUseIESBrightness(const bool& AttributeValue, bool bAddApplyDelegate = true);
+
+	UFUNCTION(BlueprintCallable, Category = "Interchange | Node | Light")
+	bool GetCustomIESBrightnessScale(float& AttributeValue) const;
+
+	UFUNCTION(BlueprintCallable, Category = "Interchange | Node | Light")
+	bool SetCustomIESBrightnessScale(const float& AttributeValue, bool bAddApplyDelegate = true);
+
+	UFUNCTION(BlueprintCallable, Category = "Interchange | Node | Light")
+	bool GetCustomRotation(FRotator& AttributeValue) const;
+
+	UFUNCTION(BlueprintCallable, Category = "Interchange | Node | Light")
+	bool SetCustomRotation(const FRotator& AttributeValue, bool bAddApplyDelegate = true);
+
+
 private:
 
 	IMPLEMENT_NODE_ATTRIBUTE_KEY(IntensityUnits)
 	IMPLEMENT_NODE_ATTRIBUTE_KEY(AttenuationRadius)
 	IMPLEMENT_NODE_ATTRIBUTE_KEY(IESTexture)
+	IMPLEMENT_NODE_ATTRIBUTE_KEY(UseIESBrightness)
+	IMPLEMENT_NODE_ATTRIBUTE_KEY(IESBrightnessScale)
+	IMPLEMENT_NODE_ATTRIBUTE_KEY(Rotation)
 };
 
 UCLASS(BlueprintType)
diff --git a/Engine/Plugins/Interchange/Runtime/Source/Pipelines/Private/InterchangeGenericScenesPipeline.cpp b/Engine/Plugins/Interchange/Runtime/Source/Pipelines/Private/InterchangeGenericScenesPipeline.cpp
--- a/Engine/Plugins/Interchange/Runtime/Source/Pipelines/Private/InterchangeGenericScenesPipeline.cpp
+++ b/Engine/Plugins/Interchange/Runtime/Source/Pipelines/Private/InterchangeGenericScenesPipeline.cpp
@@ -41,6 +41,12 @@
 
 #include UE_INLINE_GENERATED_CPP_BY_NAME(InterchangeGenericScenesPipeline)
 
+#define COPY_FROM_TRANSLATED_TO_FACTORY(TranslatedNode, FactoryNode, AttributeName, AttributeType) \
+if(AttributeType AttributeName; TranslatedNode->GetCustom##AttributeName(AttributeName)) \
+{ \
+	FactoryNode->SetCustom##AttributeName(AttributeName); \
+}
+
 namespace UE::Interchange::Private
 {
 	TArray<FString> GetAllActiveJoints(UInterchangeBaseNodeContainer* BaseNodeContainer)
@@ -685,35 +691,33 @@ void UInterchangeGenericLevelPipeline::SetUpFactoryNode(UInterchangeActorFactory
 			{
 				if (UInterchangeLightFactoryNode* LightFactoryNode = Cast<UInterchangeLightFactoryNode>(BaseLightFactoryNode))
 				{
-					if (EInterchangeLightUnits IntensityUnits; LightNode->GetCustomIntensityUnits(IntensityUnits))
+					if (FString IESTextureUid; LightNode->GetCustomIESTexture(IESTextureUid))
 					{
-						LightFactoryNode->SetCustomIntensityUnits(ELightUnits(IntensityUnits));
-					}
+						if (BaseNodeContainer->GetNode(IESTextureUid))
+						{
+							LightFactoryNode->SetCustomIESTexture(IESTextureUid);
+							LightFactoryNode->AddFactoryDependencyUid(UInterchangeFactoryBaseNode::BuildFactoryNodeUid(IESTextureUid));
 
-					if (float AttenuationRadius; LightNode->GetCustomAttenuationRadius(AttenuationRadius))
-					{
-						LightFactoryNode->SetCustomAttenuationRadius(AttenuationRadius);
+							COPY_FROM_TRANSLATED_TO_FACTORY(LightNode, LightFactoryNode, UseIESBrightness, bool)
+							COPY_FROM_TRANSLATED_TO_FACTORY(LightNode, LightFactoryNode, IESBrightnessScale, float)
+							COPY_FROM_TRANSLATED_TO_FACTORY(LightNode, LightFactoryNode, Rotation, FRotator)
+						}
 					}
 
-					if(FString IESTexture; LightNode->GetCustomIESTexture(IESTexture))
+					if (EInterchangeLightUnits IntensityUnits; LightNode->GetCustomIntensityUnits(IntensityUnits))
 					{
-						LightFactoryNode->SetCustomIESTexture(IESTexture);
+						LightFactoryNode->SetCustomIntensityUnits(ELightUnits(IntensityUnits));
 					}
 
+					COPY_FROM_TRANSLATED_TO_FACTORY(LightNode, LightFactoryNode, AttenuationRadius, float)
+
 					// RectLight
 					if(const UInterchangeRectLightNode* RectLightNode = Cast<UInterchangeRectLightNode>(LightNode))
 					{
 						if(UInterchangeRectLightFactoryNode* RectLightFactoryNode = Cast<UInterchangeRectLightFactoryNode>(LightFactoryNode))
 						{
-							if(float SourceWidth; RectLightNode->GetCustomSourceWidth(SourceWidth))
-							{
-								RectLightFactoryNode->SetCustomSourceWidth(SourceWidth);
-							}
-
-							if(float SourceHeight; RectLightNode->GetCustomSourceHeight(SourceHeight))
-							{
-								RectLightFactoryNode->SetCustomSourceHeight(SourceHeight);
-							}
+							COPY_FROM_TRANSLATED_TO_FACTORY(RectLightNode, RectLightFactoryNode, SourceWidth, float)
+							COPY_FROM_TRANSLATED_TO_FACTORY(RectLightNode, RectLightFactoryNode, SourceHeight, float)
 						}
 					}
 
@@ -726,25 +730,17 @@ void UInterchangeGenericLevelPipeline::SetUpFactoryNode(UInterchangeActorFactory
 							{
 								PointLightFactoryNode->SetCustomUseInverseSquaredFalloff(bUseInverseSquaredFalloff);
 
-								if (float LightFalloffExponent; PointLightNode->GetCustomLightFalloffExponent(LightFalloffExponent))
-								{
-									PointLightFactoryNode->SetCustomLightFalloffExponent(LightFalloffExponent);
-								}
+								COPY_FROM_TRANSLATED_TO_FACTORY(PointLightNode, PointLightFactoryNode, LightFalloffExponent, float)
 							}
 
 
 							// Spot Light
 							if (const UInterchangeSpotLightNode* SpotLightNode = Cast<UInterchangeSpotLightNode>(PointLightNode))
 							{
-								UInterchangeSpotLightFactoryNode* SpotLightFactoryNode = Cast<UInterchangeSpotLightFactoryNode>(PointLightFactoryNode);
-								if (float InnerConeAngle; SpotLightNode->GetCustomInnerConeAngle(InnerConeAngle))
-								{
-									SpotLightFactoryNode->SetCustomInnerConeAngle(InnerConeAngle);
-								}
-
-								if (float OuterConeAngle; SpotLightNode->GetCustomOuterConeAngle(OuterConeAngle))
+								if (UInterchangeSpotLightFactoryNode* SpotLightFactoryNode = Cast<UInterchangeSpotLightFactoryNode>(PointLightFactoryNode))
 								{
-									SpotLightFactoryNode->SetCustomOuterConeAngle(OuterConeAngle);
+									COPY_FROM_TRANSLATED_TO_FACTORY(SpotLightNode, SpotLightFactoryNode, InnerConeAngle, float)
+									COPY_FROM_TRANSLATED_TO_FACTORY(SpotLightNode, SpotLightFactoryNode, OuterConeAngle, float)
 								}
 							}
 						}
