Fixed missing level sequences when importing using Datasmith Interchange

diff --git a/Engine/Plugins/Experimental/Enterprise/DatasmithInterchange/Source/DatasmithInterchange/Private/InterchangeDatasmithPipeline.cpp b/Engine/Plugins/Experimental/Enterprise/DatasmithInterchange/Source/DatasmithInterchange/Private/InterchangeDatasmithPipeline.cpp
--- a/Engine/Plugins/Experimental/Enterprise/DatasmithInterchange/Source/DatasmithInterchange/Private/InterchangeDatasmithPipeline.cpp
+++ b/Engine/Plugins/Experimental/Enterprise/DatasmithInterchange/Source/DatasmithInterchange/Private/InterchangeDatasmithPipeline.cpp
@@ -147,7 +147,11 @@ void UInterchangeDatasmithPipeline::AdjustSettingsForContext(EInterchangePipelin
 
 void UInterchangeDatasmithPipeline::PostDuplicate(bool bDuplicateForPIE)
 {
-	AdjustSettingsForContext(CachePipelineContext, CacheReimportObject.Get());
+	// Only adjust settings if there is anything cached.
+	if (CachePipelineContext != EInterchangePipelineContext::None)
+	{
+		AdjustSettingsForContext(CachePipelineContext, CacheReimportObject.Get());
+	}
 }
 
 void UInterchangeDatasmithPipeline::ExecutePipeline(UInterchangeBaseNodeContainer* InBaseNodeContainer, const TArray<UInterchangeSourceData*>& SourceDatas, const FString& ContentBasePath)
