- Fixed Dataprep crash
- Added proper handling of Datasmith settings thru Interchange
- Remove useless folder when importing thru Interchange in Dataprep
- Eliminate options dialog popup in Dataprep with Interchange file producer

#rnx

diff --git a/Engine/Plugins/Enterprise/DataprepEditor/Source/DataprepCore/Private/InterchangeDataprepPipeline.cpp b/Engine/Plugins/Enterprise/DataprepEditor/Source/DataprepCore/Private/InterchangeDataprepPipeline.cpp
--- a/Engine/Plugins/Enterprise/DataprepEditor/Source/DataprepCore/Private/InterchangeDataprepPipeline.cpp
+++ b/Engine/Plugins/Enterprise/DataprepEditor/Source/DataprepCore/Private/InterchangeDataprepPipeline.cpp
@@ -33,34 +33,34 @@ void UInterchangeDataprepLevelPipeline::ExecutePipeline(UInterchangeBaseNodeCont
 		}
 		else if (UInterchangeTextureFactoryNode* TextureFactoryNode = Cast<UInterchangeTextureFactoryNode>(FactoryNode))
 		{
-			TextureFactoryNode->SetCustomSubPath(FPaths::Combine(PackageSubPath, "Textures"));
+			TextureFactoryNode->SetCustomSubPath(TEXT("Textures"));
 		}
 		else if (UInterchangeBaseMaterialFactoryNode* MaterialFactoryNode = Cast<UInterchangeBaseMaterialFactoryNode>(FactoryNode))
 		{
 			if (MaterialFactoryNode->IsA<UInterchangeMaterialFactoryNode>())
 			{
-				MaterialFactoryNode->SetCustomSubPath(FPaths::Combine(PackageSubPath, "Materials/References"));
+				MaterialFactoryNode->SetCustomSubPath(TEXT("Materials/References"));
 			}
 			else if (MaterialFactoryNode->IsA<UInterchangeMaterialFunctionFactoryNode>())
 			{
-				MaterialFactoryNode->SetCustomSubPath(FPaths::Combine(PackageSubPath, "Materials/References/Functions"));
+				MaterialFactoryNode->SetCustomSubPath(TEXT("Materials/References/Functions"));
 			}
 			else
 			{
-				MaterialFactoryNode->SetCustomSubPath(FPaths::Combine(PackageSubPath, "Materials"));
+				MaterialFactoryNode->SetCustomSubPath(TEXT("Materials"));
 			}
 		}
 		else if (UInterchangeStaticMeshFactoryNode* MeshFactoryNode = Cast<UInterchangeStaticMeshFactoryNode>(FactoryNode))
 		{
-			MeshFactoryNode->SetCustomSubPath(FPaths::Combine(PackageSubPath, "Geometries"));
+			MeshFactoryNode->SetCustomSubPath(TEXT("Geometries"));
 		}
 		else if (UInterchangeLevelSequenceFactoryNode* SequenceFactoryNode = Cast<UInterchangeLevelSequenceFactoryNode>(FactoryNode))
 		{
-			SequenceFactoryNode->SetCustomSubPath(FPaths::Combine(PackageSubPath, "Animations"));
+			SequenceFactoryNode->SetCustomSubPath(TEXT("Animations"));
 		}
 		else if (UInterchangeSceneVariantSetsFactoryNode* VariantFactoryNode = Cast<UInterchangeSceneVariantSetsFactoryNode>(FactoryNode))
 		{
-			VariantFactoryNode->SetCustomSubPath(FPaths::Combine(PackageSubPath, "Variants"));
+			VariantFactoryNode->SetCustomSubPath(TEXT("Variants"));
 		}
 		// Remove UInterchangeSceneImportAssetFactoryNode from the import
 		// #dataprep_todo : Revisit when consumer will have been updated.
@@ -83,6 +83,6 @@ void UInterchangeDataprepLevelPipeline::ExecutePostImportPipeline(const UInterch
 
 	Super::ExecutePostImportPipeline(NodeContainer, FactoryNodeKey, CreatedAsset, bIsAReimport);
 
-	CreatedAsset->ClearFlags(RF_Public/* | RF_Standalone*/);
-	CreatedAsset->SetFlags(RF_Transient);
+	CreatedAsset->ClearFlags(RF_Public);
+	CreatedAsset->SetFlags(RF_Transactional|RF_Standalone);
 }
diff --git a/Engine/Plugins/Enterprise/DatasmithImporter/Source/DatasmithImporter/Private/DatasmithConsumer.cpp b/Engine/Plugins/Enterprise/DatasmithImporter/Source/DatasmithImporter/Private/DatasmithConsumer.cpp
--- a/Engine/Plugins/Enterprise/DatasmithImporter/Source/DatasmithImporter/Private/DatasmithConsumer.cpp
+++ b/Engine/Plugins/Enterprise/DatasmithImporter/Source/DatasmithImporter/Private/DatasmithConsumer.cpp
@@ -1159,16 +1159,6 @@ bool UDatasmithConsumer::ValidateAssets()
 	{
 		if(UObject* Asset = AssetPtr.Get())
 		{
-			// If the asset has not been created by a Datasmith producer
-			const FString SourcePackagePath = Asset->GetOutermost()->GetName();
-			if (SourcePackagePath.StartsWith(Context.TransientContentFolder, ESearchCase::CaseSensitive))
-			{
-				const FString DestinationPackagePath = SourcePackagePath.Replace(*Context.TransientContentFolder, *TargetContentFolder, ESearchCase::CaseSensitive);
-
-				Asset->ClearFlags(RF_Transient);
-				Asset = FDatasmithImporterImpl::FinalizeAsset(Asset, *DestinationPackagePath, nullptr, nullptr);
-			}
-
 			TPair<FString, FString>& AssetPackageInfo = AssetPackageInfoMap.Add(Asset);
 
 			bool bAssetWithMarker = false;
diff --git a/Engine/Plugins/Experimental/Enterprise/DatasmithInterchange/Source/DatasmithInterchange/DatasmithInterchange.Build.cs b/Engine/Plugins/Experimental/Enterprise/DatasmithInterchange/Source/DatasmithInterchange/DatasmithInterchange.Build.cs
--- a/Engine/Plugins/Experimental/Enterprise/DatasmithInterchange/Source/DatasmithInterchange/DatasmithInterchange.Build.cs
+++ b/Engine/Plugins/Experimental/Enterprise/DatasmithInterchange/Source/DatasmithInterchange/DatasmithInterchange.Build.cs
@@ -10,6 +10,7 @@ namespace UnrealBuildTool.Rules
 			PrivateDependencyModuleNames.AddRange(
 				new string[]
 				{
+					"CADTools",
 					"DatasmithContent",
 					"DatasmithCore",
 					"DatasmithTranslator",
diff --git a/Engine/Plugins/Experimental/Enterprise/DatasmithInterchange/Source/DatasmithInterchange/Private/InterchangeDatasmithTranslator.cpp b/Engine/Plugins/Experimental/Enterprise/DatasmithInterchange/Source/DatasmithInterchange/Private/InterchangeDatasmithTranslator.cpp
--- a/Engine/Plugins/Experimental/Enterprise/DatasmithInterchange/Source/DatasmithInterchange/Private/InterchangeDatasmithTranslator.cpp
+++ b/Engine/Plugins/Experimental/Enterprise/DatasmithInterchange/Source/DatasmithInterchange/Private/InterchangeDatasmithTranslator.cpp
@@ -18,6 +18,7 @@
 #include "DatasmithVariantElements.h"
 #include "IDatasmithSceneElements.h"
 
+#include "CADOptions.h"
 #include "ExternalSourceModule.h"
 #include "SourceUri.h"
 #include "InterchangeCameraNode.h"
@@ -168,54 +169,34 @@ bool UInterchangeDatasmithTranslator::Translate(UInterchangeBaseNodeContainer& B
 		return false;
 	}
 
-	// Temporary: Update the extra options of the associated translator
-#if WITH_EDITOR
+	StartTime = FPlatformTime::Cycles64();
+	FPaths::NormalizeFilename(FilePath);
+
+	TSharedPtr<IDatasmithScene> DatasmithScene;
 	{
-		const TSharedPtr<IDatasmithTranslator>& DatasmithTranslator = LoadedExternalSource->GetAssetTranslator();
-		if (!DatasmithTranslator)
-		{
-			return false;
-		}
+		TGuardValue<bool> EnableCADCache(CADLibrary::FImportParameters::bGEnableCADCache, true);
 
-		bool bShouldImport = true;
-		if (!IsInGameThread())
+		if (GetSettings())
 		{
-			TSharedRef<TPromise<bool>, ESPMode::ThreadSafe> Promise = MakeShareable(new TPromise<bool>());
-			TFunction<void()> PromiseKeeper = [&DatasmithTranslator, &Promise]() -> void
-			{
-				bool bShouldImport = UE::DatasmithInterchange::DisplayOptionsDialog(*DatasmithTranslator);
-				Promise->SetValue(bShouldImport);
-			};
-
-			AsyncTask(ENamedThreads::GameThread, MoveTemp(PromiseKeeper));
+			CADLibrary::FImportParameters::bGEnableCADCache = true;
 
-			TFuture<bool> Future = Promise->GetFuture();
-			
-			Future.Wait();
-			
-			bShouldImport = Future.Get();
-		}
-		else
-		{
-			bShouldImport = UE::DatasmithInterchange::DisplayOptionsDialog(*DatasmithTranslator);
+			const TSharedPtr<IDatasmithTranslator>& DatasmithTranslator = LoadedExternalSource->GetAssetTranslator();
+			if (DatasmithTranslator)
+			{
+				TArray<TObjectPtr<UDatasmithOptionsBase>> OptionArray;
+				OptionArray.Add(CachedSettings->ImportOptions);
+				DatasmithTranslator->SetSceneImportOptions(OptionArray);
+			}
 		}
 
-		if (!bShouldImport)
+		// Should it be mutable instead? If Translate is const should we really be doing this?.
+		DatasmithScene = LoadedExternalSource->TryLoad();
+
+		if (!DatasmithScene.IsValid())
 		{
 			return false;
 		}
 	}
-#endif
-
-	StartTime = FPlatformTime::Cycles64();
-	FPaths::NormalizeFilename(FilePath);
-
-	// Should it be mutable instead? If Translate is const should we really be doing this?.
-	TSharedPtr<IDatasmithScene> DatasmithScene = LoadedExternalSource->TryLoad();
-	if (!DatasmithScene.IsValid())
-	{
-		return false;
-	}
 
 	// Texture Nodes
 	{
@@ -438,6 +419,7 @@ void UInterchangeDatasmithTranslator::HandleDatasmithActor(UInterchangeBaseNodeC
 
 	UInterchangeSceneNode* InterchangeSceneNode = NewObject<UInterchangeSceneNode>(&BaseNodeContainer);
 	InterchangeSceneNode->InitializeNode(NodeUid, ActorElement->GetLabel(), EInterchangeNodeContainerType::TranslatedScene);
+	InterchangeSceneNode->SetAssetName(NodeName);
 	BaseNodeContainer.AddNode(InterchangeSceneNode);
 	BaseNodeContainer.SetNodeParentUid(NodeUid, ParentNodeUid);
 
@@ -942,4 +924,61 @@ void UInterchangeDatasmithTranslator::ImportFinish()
 }
 
 
+UInterchangeTranslatorSettings* UInterchangeDatasmithTranslator::GetSettings() const
+{
+	using namespace UE::DatasmithImporter;
+	using namespace UE::DatasmithInterchange;
+
+	if (!CachedSettings)
+	{
+		if (!LoadedExternalSource.IsValid())
+		{
+			FString FilePath = FPaths::ConvertRelativePathToFull(SourceData->GetFilename());
+			FileName = FPaths::GetCleanFilename(FilePath);
+			const FSourceUri FileNameUri = FSourceUri::FromFilePath(FilePath);
+			LoadedExternalSource = IExternalSourceModule::GetOrCreateExternalSource(FileNameUri);
+		}
+
+		if (!LoadedExternalSource.IsValid() || !LoadedExternalSource->IsAvailable())
+		{
+			return nullptr;
+		}
+
+		const TSharedPtr<IDatasmithTranslator>& DatasmithTranslator = LoadedExternalSource->GetAssetTranslator();
+		if (!DatasmithTranslator)
+		{
+			return nullptr;
+		}
+
+		TArray<TObjectPtr<UDatasmithOptionsBase>> OptionArray;
+		DatasmithTranslator->GetSceneImportOptions(OptionArray);
+		if (OptionArray.Num() == 0)
+		{
+			return nullptr;
+		}
+
+		CachedSettings = DuplicateObject<UInterchangeDatasmithTranslatorSettings>(UInterchangeDatasmithTranslatorSettings::StaticClass()->GetDefaultObject<UInterchangeDatasmithTranslatorSettings>(), GetTransientPackage());
+		CachedSettings->SetFlags(RF_Standalone);
+		CachedSettings->ClearInternalFlags(EInternalObjectFlags::Async);
+		CachedSettings->ImportOptions = OptionArray[0];
+	}
+	return CachedSettings;
+}
+
+void UInterchangeDatasmithTranslator::SetSettings(const UInterchangeTranslatorSettings* InterchangeTranslatorSettings)
+{
+	if (CachedSettings)
+	{
+		CachedSettings->ClearFlags(RF_Standalone);
+		CachedSettings->ClearInternalFlags(EInternalObjectFlags::Async);
+		CachedSettings = nullptr;
+	}
+	if (InterchangeTranslatorSettings)
+	{
+		CachedSettings = DuplicateObject<UInterchangeDatasmithTranslatorSettings>(Cast<UInterchangeDatasmithTranslatorSettings>(InterchangeTranslatorSettings), GetTransientPackage());
+		CachedSettings->ClearInternalFlags(EInternalObjectFlags::Async);
+		CachedSettings->SetFlags(RF_Standalone);
+	}
+}
+
 #undef LOCTEXT_NAMESPACE
\ No newline at end of file
diff --git a/Engine/Plugins/Experimental/Enterprise/DatasmithInterchange/Source/DatasmithInterchange/Public/InterchangeDatasmithTranslator.h b/Engine/Plugins/Experimental/Enterprise/DatasmithInterchange/Source/DatasmithInterchange/Public/InterchangeDatasmithTranslator.h
--- a/Engine/Plugins/Experimental/Enterprise/DatasmithInterchange/Source/DatasmithInterchange/Public/InterchangeDatasmithTranslator.h
+++ b/Engine/Plugins/Experimental/Enterprise/DatasmithInterchange/Source/DatasmithInterchange/Public/InterchangeDatasmithTranslator.h
@@ -29,6 +29,7 @@ class IDatasmithLightActorElement;
 class IDatasmithDecalActorElement;
 class IDatasmithScene;
 class IDatasmithTransformAnimationElement;
+class UDatasmithOptionsBase;
 class UInterchangePhysicalCameraNode;
 class UInterchangeBaseLightNode;
 class UInterchangeDecalNode;
@@ -54,6 +55,16 @@ namespace UE::DatasmithInterchange::AnimUtils
 	extern bool GetAnimationPayloadData(const IDatasmithBaseAnimationElement& AnimationElement, float FrameRate, EInterchangeAnimationPayLoadType PayLoadType, UE::Interchange::FAnimationPayloadData& PayLoadData);
 }
 
+UCLASS(BlueprintType, editinlinenew, MinimalAPI)
+class UInterchangeDatasmithTranslatorSettings : public UInterchangeTranslatorSettings
+{
+	GENERATED_BODY()
+
+public:
+	UPROPERTY(EditAnywhere, Instanced, BlueprintReadWrite, Category = "Datasmith Interchange", meta = (ShowOnlyInnerProperties))
+	TObjectPtr<UDatasmithOptionsBase> ImportOptions;
+};
+
 UCLASS(BlueprintType, Experimental)
 class DATASMITHINTERCHANGE_API UInterchangeDatasmithTranslator : public UInterchangeTranslatorBase
 	, public IInterchangeTexturePayloadInterface
@@ -66,6 +77,7 @@ class DATASMITHINTERCHANGE_API UInterchangeDatasmithTranslator : public UInterch
 
 public:
 
+	/** Begin UInterchangeTranslatorBase API*/
 	virtual bool CanImportSourceData(const UInterchangeSourceData* InSourceData) const override;
 
 	virtual bool Translate(UInterchangeBaseNodeContainer& BaseNodeContainer) const override;
@@ -87,6 +99,10 @@ public:
 
 	virtual void ImportFinish() override;
 
+	virtual UInterchangeTranslatorSettings* GetSettings() const override;
+	virtual void SetSettings(const UInterchangeTranslatorSettings* InterchangeTranslatorSettings) override;
+	/** End UInterchangeTranslatorBase API*/
+
 	/* IInterchangeTexturePayloadInterface Begin */
 	virtual TOptional<UE::Interchange::FImportImage> GetTexturePayloadData(const FString& PayloadKey, TOptional<FString>& AlternateTexturePath) const override;
 	/* IInterchangeTexturePayloadInterface End */
@@ -123,6 +139,7 @@ private:
 
 	mutable uint64 StartTime = 0;
 	mutable FString FileName;
+	mutable TObjectPtr<UInterchangeDatasmithTranslatorSettings> CachedSettings = nullptr;
 
 	mutable TMap<FString, UE::DatasmithInterchange::AnimUtils::FAnimationPayloadDesc> AnimationPayLoadMapping;
-};
\ No newline at end of file
+};
diff --git a/Engine/Plugins/Interchange/Editor/Source/Pipelines/Private/SInterchangePipelineConfigurationDialog.cpp b/Engine/Plugins/Interchange/Editor/Source/Pipelines/Private/SInterchangePipelineConfigurationDialog.cpp
--- a/Engine/Plugins/Interchange/Editor/Source/Pipelines/Private/SInterchangePipelineConfigurationDialog.cpp
+++ b/Engine/Plugins/Interchange/Editor/Source/Pipelines/Private/SInterchangePipelineConfigurationDialog.cpp
@@ -259,25 +259,22 @@ void SInterchangePipelineConfigurationDialog::OnFinishedChangingProperties(const
 	}
 	if (UClass* TranslatorSettingsClass = TranslatorSettings->GetClass())
 	{
-		if (TranslatorSettingsClass->HasProperty(PropertyChangedEvent.Property))
-		{
-			//Save the config locally before the translation.
-			TranslatorSettings->SaveSettings();
+		//Save the config locally before the translation.
+		TranslatorSettings->SaveSettings();
 
-			//Need to Translate the source data
- 			FScopedSlowTask Progress(2.f, NSLOCTEXT("SInterchangePipelineConfigurationDialog", "TranslatingSourceFile...", "Translating source file..."));
- 			Progress.MakeDialog();
- 			Progress.EnterProgressFrame(1.f);
-			//Reset the container
-			BaseNodeContainer->Reset();
+		//Need to Translate the source data
+		FScopedSlowTask Progress(2.f, NSLOCTEXT("SInterchangePipelineConfigurationDialog", "TranslatingSourceFile...", "Translating source file..."));
+		Progress.MakeDialog();
+		Progress.EnterProgressFrame(1.f);
+		//Reset the container
+		BaseNodeContainer->Reset();
 
-			Translator->Translate(*BaseNodeContainer.Get());
+		Translator->Translate(*BaseNodeContainer.Get());
 
-			//Refresh the dialog
-			RefreshStack(false);
+		//Refresh the dialog
+		RefreshStack(false);
 
-			Progress.EnterProgressFrame(1.f);
-		}
+		Progress.EnterProgressFrame(1.f);
 	}
 }
 
