Fixed crash when using Datasmith Intercchange with some CAD files

diff --git a/Engine/Plugins/Experimental/Enterprise/DatasmithInterchange/Source/DatasmithInterchange/Private/InterchangeDatasmithTranslator.cpp b/Engine/Plugins/Experimental/Enterprise/DatasmithInterchange/Source/DatasmithInterchange/Private/InterchangeDatasmithTranslator.cpp
--- a/Engine/Plugins/Experimental/Enterprise/DatasmithInterchange/Source/DatasmithInterchange/Private/InterchangeDatasmithTranslator.cpp
+++ b/Engine/Plugins/Experimental/Enterprise/DatasmithInterchange/Source/DatasmithInterchange/Private/InterchangeDatasmithTranslator.cpp
@@ -183,6 +183,14 @@ bool UInterchangeDatasmithTranslator::Translate(UInterchangeBaseNodeContainer& B
 			const TSharedPtr<IDatasmithTranslator>& DatasmithTranslator = LoadedExternalSource->GetAssetTranslator();
 			if (DatasmithTranslator)
 			{
+				FDatasmithTranslatorCapabilities Capabilities;
+				DatasmithTranslator->Initialize(Capabilities);
+
+				if (!Capabilities.bParallelLoadStaticMeshSupported)
+				{
+					AsyncMode = EAsyncExecution::TaskGraphMainThread;
+				}
+
 				TArray<TObjectPtr<UDatasmithOptionsBase>> OptionArray;
 				OptionArray.Add(CachedSettings->ImportOptions);
 				DatasmithTranslator->SetSceneImportOptions(OptionArray);
@@ -784,7 +792,7 @@ TFuture<TOptional<UE::Interchange::FMeshPayloadData>> UInterchangeDatasmithTrans
 		return EmptyPromise.GetFuture();
 	}
 
-	return Async(EAsyncExecution::TaskGraph, [this, MeshElement = MoveTemp(MeshElement), MeshGlobalTransform]
+	return Async(AsyncMode, [this, MeshElement = MoveTemp(MeshElement), MeshGlobalTransform]
 		{
 			TOptional<UE::Interchange::FMeshPayloadData> Result;
 
diff --git a/Engine/Plugins/Experimental/Enterprise/DatasmithInterchange/Source/DatasmithInterchange/Public/InterchangeDatasmithTranslator.h b/Engine/Plugins/Experimental/Enterprise/DatasmithInterchange/Source/DatasmithInterchange/Public/InterchangeDatasmithTranslator.h
--- a/Engine/Plugins/Experimental/Enterprise/DatasmithInterchange/Source/DatasmithInterchange/Public/InterchangeDatasmithTranslator.h
+++ b/Engine/Plugins/Experimental/Enterprise/DatasmithInterchange/Source/DatasmithInterchange/Public/InterchangeDatasmithTranslator.h
@@ -143,4 +143,5 @@ private:
 	mutable TObjectPtr<UInterchangeDatasmithTranslatorSettings> CachedSettings = nullptr;
 
 	mutable TMap<FString, UE::DatasmithInterchange::AnimUtils::FAnimationPayloadDesc> AnimationPayLoadMapping;
+	mutable EAsyncExecution AsyncMode = EAsyncExecution::TaskGraph;
 };
