Added retessellation action for meshes imported through Datasmith Interchange from CAD files.

#rnx

diff --git a/Engine/Plugins/Enterprise/DatasmithCADImporter/Config/BaseDatasmithCADImporter.ini b/Engine/Plugins/Enterprise/DatasmithCADImporter/Config/BaseDatasmithCADImporter.ini
--- a/Engine/Plugins/Enterprise/DatasmithCADImporter/Config/BaseDatasmithCADImporter.ini
+++ b/Engine/Plugins/Enterprise/DatasmithCADImporter/Config/BaseDatasmithCADImporter.ini
@@ -1,5 +1,4 @@
 [CoreRedirects]
-
 +ClassRedirects=(OldName="/Script/DatasmithCoreTechExtension.CoreTechBlueprintLibrary", NewName="/Script/ParametricSurfaceExtension.ParametricSurfaceBlueprintLibrary")
 +ClassRedirects=(OldName="/Script/CoreTechExtension.CoreTechBlueprintLibrary", NewName="/Script/ParametricSurfaceExtension.ParametricSurfaceBlueprintLibrary")
 
@@ -12,4 +11,11 @@
 +ClassRedirects=(OldName="/Script/DatasmithCoreTechExtension.DataprepTessellationOperation", NewName="/Script/ParametricSurfaceExtension.DataprepTessellationOperation")
 +ClassRedirects=(OldName="/Script/CoreTechExtension.DataprepTessellationOperation", NewName="/Script/ParametricSurfaceExtension.DataprepTessellationOperation")
 
++ClassRedirects=(OldName="/Script/ParametricSurface.UParametricSurfaceData", NewName="/Script/ParametricSurface.UDatasmithParametricSurfaceData")
+
+[SectionsToSave]
++Section=/Script/DatasmithWireTranslator.DatasmithWireOptions
++Section=/Script/DatasmithOpenNurbsTranslator.DatasmithOpenNurbsImportOptions
++Section=/Script/DatasmithOpenNurbsTranslator.DatasmithOpenNurbsImportOptions
+
 
diff --git a/Engine/Plugins/Enterprise/DatasmithCADImporter/Source/CADKernelSurface/Public/CADKernelSurfaceExtension.h b/Engine/Plugins/Enterprise/DatasmithCADImporter/Source/CADKernelSurface/Public/CADKernelSurfaceExtension.h
--- a/Engine/Plugins/Enterprise/DatasmithCADImporter/Source/CADKernelSurface/Public/CADKernelSurfaceExtension.h
+++ b/Engine/Plugins/Enterprise/DatasmithCADImporter/Source/CADKernelSurface/Public/CADKernelSurfaceExtension.h
@@ -4,7 +4,7 @@
 
 #include "CoreMinimal.h"
 
-#include "ParametricSurfaceData.h"
+#include "DatasmithParametricSurfaceData.h"
 #include "DatasmithAdditionalData.h"
 #include "DatasmithCustomAction.h"
 #include "DatasmithImportOptions.h"
@@ -13,7 +13,7 @@
 #include "CADKernelSurfaceExtension.generated.h"
 
 UCLASS(meta = (DisplayName = "CADKernel Parametric Surface Data"))
-class CADKERNELSURFACE_API UCADKernelParametricSurfaceData : public UParametricSurfaceData
+class CADKERNELSURFACE_API UCADKernelParametricSurfaceData : public UDatasmithParametricSurfaceData
 {
 	GENERATED_BODY()
 
diff --git a/Engine/Plugins/Enterprise/DatasmithCADImporter/Source/ParametricSurface/ParametricSurface.Build.cs b/Engine/Plugins/Enterprise/DatasmithCADImporter/Source/ParametricSurface/ParametricSurface.Build.cs
--- a/Engine/Plugins/Enterprise/DatasmithCADImporter/Source/ParametricSurface/ParametricSurface.Build.cs
+++ b/Engine/Plugins/Enterprise/DatasmithCADImporter/Source/ParametricSurface/ParametricSurface.Build.cs
@@ -21,6 +21,7 @@ namespace UnrealBuildTool.Rules
 					"Core",
 					"CoreUObject",
 					"CADInterfaces",
+					"CADKernel",
 					"CADLibrary",
 					"CADTools",
 					"DatasmithContent",
diff --git a/Engine/Plugins/Enterprise/DatasmithCADImporter/Source/ParametricSurface/Private/DatasmithParametricSurfaceData.cpp b/Engine/Plugins/Enterprise/DatasmithCADImporter/Source/ParametricSurface/Private/DatasmithParametricSurfaceData.cpp
new file mode 100644
--- /dev/null
+++ b/Engine/Plugins/Enterprise/DatasmithCADImporter/Source/ParametricSurface/Private/DatasmithParametricSurfaceData.cpp
@@ -0,0 +1,59 @@
+// Copyright Epic Games, Inc. All Rights Reserved.
+
+#include "DatasmithParametricSurfaceData.h"
+
+#include "CADOptions.h"
+
+#include "Misc/FileHelper.h"
+#include "UObject/EnterpriseObjectVersion.h"
+
+#define LOCTEXT_NAMESPACE "ParametricSurfaceData"
+
+bool UDatasmithParametricSurfaceData::SetFile(const TCHAR* FilePath)
+{
+	if (FPaths::FileExists(FilePath))
+	{
+		TArray<uint8> ByteArray;
+
+		if (FFileHelper::LoadFileToArray(ByteArray, FilePath))
+		{
+			RawData = MoveTemp(ByteArray);
+			return true;
+		}
+	}
+
+	return false;
+}
+
+void UDatasmithParametricSurfaceData::SetImportParameters(const CADLibrary::FImportParameters& InSceneParameters)
+{
+	SceneParameters.ModelCoordSys = uint8(InSceneParameters.GetModelCoordSys());
+}
+
+void UDatasmithParametricSurfaceData::SetMeshParameters(const CADLibrary::FMeshParameters& InMeshParameters)
+{
+	MeshParameters.bNeedSwapOrientation = InMeshParameters.bNeedSwapOrientation;
+	MeshParameters.bIsSymmetric = InMeshParameters.bIsSymmetric;
+	MeshParameters.SymmetricNormal = (FVector) InMeshParameters.SymmetricNormal;
+	MeshParameters.SymmetricOrigin = (FVector) InMeshParameters.SymmetricOrigin;
+}
+
+void UDatasmithParametricSurfaceData::Serialize(FArchive& Ar)
+{
+	Ar.UsingCustomVersion(FEnterpriseObjectVersion::GUID);
+
+	Super::Serialize(Ar);
+
+	if (Ar.IsSaving() || (Ar.IsLoading() && Ar.CustomVer(FEnterpriseObjectVersion::GUID) >= FEnterpriseObjectVersion::CoreTechParametricSurfaceOptim))
+	{
+		Ar << RawData;
+	}
+
+	if (RawData_DEPRECATED.Num() && RawData.Num() == 0)
+	{
+		RawData = MoveTemp(RawData_DEPRECATED);
+	}
+}
+
+#undef LOCTEXT_NAMESPACE // "ParametricSurfaceData"
+
diff --git a/Engine/Plugins/Enterprise/DatasmithCADImporter/Source/ParametricSurface/Private/ParametricSurfaceData.cpp b/Engine/Plugins/Enterprise/DatasmithCADImporter/Source/ParametricSurface/Private/ParametricSurfaceData.cpp
deleted file mode 100644
diff --git a/Engine/Plugins/Enterprise/DatasmithCADImporter/Source/ParametricSurface/Private/ParametricSurfaceModule.cpp b/Engine/Plugins/Enterprise/DatasmithCADImporter/Source/ParametricSurface/Private/ParametricSurfaceModule.cpp
--- a/Engine/Plugins/Enterprise/DatasmithCADImporter/Source/ParametricSurface/Private/ParametricSurfaceModule.cpp
+++ b/Engine/Plugins/Enterprise/DatasmithCADImporter/Source/ParametricSurface/Private/ParametricSurfaceModule.cpp
@@ -15,7 +15,7 @@ void FParametricSurfaceModule::StartupModule()
 	{
 		TArray<FCoreRedirect> Redirects;
 		Redirects.Emplace(ECoreRedirectFlags::Type_Package, TEXT("/Script/DatasmithCoreTechParametricSurfaceData"), TEXT("/Script/ParametricSurface"));
-		Redirects.Emplace(ECoreRedirectFlags::Type_Property, TEXT("UParametricSurfaceData.RawData"), TEXT("RawData_DEPRECATED"));
+		Redirects.Emplace(ECoreRedirectFlags::Type_Property, TEXT("UDatasmithParametricSurfaceData.RawData"), TEXT("RawData_DEPRECATED"));
 		Redirects.Emplace(ECoreRedirectFlags::Type_Property, TEXT("UCoreTechParametricSurfaceData.RawData"), TEXT("RawData_DEPRECATED"));
 		Redirects.Emplace(ECoreRedirectFlags::Type_Struct, TEXT("CoreTechSceneParameters"), TEXT("/Script/ParametricSurface.ParametricSceneParameters"));
 		Redirects.Emplace(ECoreRedirectFlags::Type_Struct, TEXT("CoreTechMeshParameters"), TEXT("/Script/ParametricSurface.ParametricMeshParameters"));
@@ -33,7 +33,7 @@ bool FParametricSurfaceModule::IsAvailable()
 	return FModuleManager::Get().IsModuleLoaded(PARAMETRICSURFACE_MODULE_NAME);
 }
 
-UParametricSurfaceData* FParametricSurfaceModule::CreateParametricSurface()
+UDatasmithParametricSurfaceData* FParametricSurfaceModule::CreateParametricSurface()
 {
 	return Datasmith::MakeAdditionalData<UTechSoftParametricSurfaceData>();
 }
diff --git a/Engine/Plugins/Enterprise/DatasmithCADImporter/Source/ParametricSurface/Private/ParametricSurfaceTranslator.cpp b/Engine/Plugins/Enterprise/DatasmithCADImporter/Source/ParametricSurface/Private/ParametricSurfaceTranslator.cpp
--- a/Engine/Plugins/Enterprise/DatasmithCADImporter/Source/ParametricSurface/Private/ParametricSurfaceTranslator.cpp
+++ b/Engine/Plugins/Enterprise/DatasmithCADImporter/Source/ParametricSurface/Private/ParametricSurfaceTranslator.cpp
@@ -2,7 +2,7 @@
 
 #include "ParametricSurfaceTranslator.h"
 
-#include "ParametricSurfaceData.h"
+#include "DatasmithParametricSurfaceData.h"
 #include "ParametricSurfaceModule.h"
 
 #include "CADOptions.h"
@@ -48,7 +48,7 @@ bool ParametricSurfaceUtils::AddSurfaceData(const TCHAR* MeshFilePath, const CAD
 {
 	if (MeshFilePath && IFileManager::Get().FileExists(MeshFilePath))
 	{
-		UParametricSurfaceData* ParametricSurfaceData = FParametricSurfaceModule::CreateParametricSurface();
+		UDatasmithParametricSurfaceData* ParametricSurfaceData = FParametricSurfaceModule::CreateParametricSurface();
 
 		if (!ParametricSurfaceData || !ParametricSurfaceData->SetFile(MeshFilePath))
 		{
@@ -61,12 +61,6 @@ bool ParametricSurfaceUtils::AddSurfaceData(const TCHAR* MeshFilePath, const CAD
 
 		OutMeshPayload.AdditionalData.Add(ParametricSurfaceData);
 
-		// Remove the file because it is temporary since caching is disabled.
-		if (!CADLibrary::FImportParameters::bGEnableCADCache)
-		{
-			IFileManager::Get().Delete(MeshFilePath);
-		}
-
 		return true;
 	}
 
diff --git a/Engine/Plugins/Enterprise/DatasmithCADImporter/Source/ParametricSurface/Private/TechSoft/TechSoftParametricSurface.cpp b/Engine/Plugins/Enterprise/DatasmithCADImporter/Source/ParametricSurface/Private/TechSoft/TechSoftParametricSurface.cpp
--- a/Engine/Plugins/Enterprise/DatasmithCADImporter/Source/ParametricSurface/Private/TechSoft/TechSoftParametricSurface.cpp
+++ b/Engine/Plugins/Enterprise/DatasmithCADImporter/Source/ParametricSurface/Private/TechSoft/TechSoftParametricSurface.cpp
@@ -3,6 +3,7 @@
 #include "TechSoftParametricSurface.h"
 
 #include "CADInterfacesModule.h"
+#include "CADOptions.h"
 #include "DatasmithAdditionalData.h"
 #include "DatasmithPayload.h"
 #include "TechSoftInterface.h"
@@ -38,7 +39,8 @@ bool UTechSoftParametricSurfaceData::Tessellate(UStaticMesh& StaticMesh, const F
 		CADLibrary::FImportParameters ImportParameters((FDatasmithUtils::EModelCoordSystem)SceneParameters.ModelCoordSys);
 		ImportParameters.SetTesselationParameters(RetessellateOptions.ChordTolerance, RetessellateOptions.MaxEdgeLength, RetessellateOptions.NormalTolerance, (CADLibrary::EStitchingTechnique)RetessellateOptions.StitchingTechnique);
 
-		FMeshConversionContext Context(ImportParameters, MeshParameters, RetessellateOptions.GetGeometricTolerance(true));
+		CADLibrary::FMeshParameters Parameters{ MeshParameters.bNeedSwapOrientation, MeshParameters.bIsSymmetric, (FVector3f)MeshParameters.SymmetricNormal, (FVector3f)MeshParameters.SymmetricOrigin };
+		FMeshConversionContext Context(ImportParameters, Parameters, RetessellateOptions.GetGeometricTolerance(true));
 
 		CADLibrary::FTechSoftInterface& TechSoftInterface = CADLibrary::FTechSoftInterface::Get();
 		bSuccessfulTessellation = TechSoftInterface.InitializeKernel(*FPaths::EnginePluginsDir());
diff --git a/Engine/Plugins/Enterprise/DatasmithCADImporter/Source/ParametricSurface/Private/TechSoft/TechSoftParametricSurface.h b/Engine/Plugins/Enterprise/DatasmithCADImporter/Source/ParametricSurface/Private/TechSoft/TechSoftParametricSurface.h
--- a/Engine/Plugins/Enterprise/DatasmithCADImporter/Source/ParametricSurface/Private/TechSoft/TechSoftParametricSurface.h
+++ b/Engine/Plugins/Enterprise/DatasmithCADImporter/Source/ParametricSurface/Private/TechSoft/TechSoftParametricSurface.h
@@ -4,7 +4,7 @@
 
 #include "CoreMinimal.h"
 
-#include "ParametricSurfaceData.h"
+#include "DatasmithParametricSurfaceData.h"
 #include "DatasmithAdditionalData.h"
 #include "DatasmithCustomAction.h"
 #include "DatasmithImportOptions.h"
@@ -13,7 +13,7 @@
 #include "TechSoftParametricSurface.generated.h"
 
 UCLASS(meta = (DisplayName = "TechSoft Parametric Surface Data"))
-class PARAMETRICSURFACE_API UTechSoftParametricSurfaceData : public UParametricSurfaceData
+class PARAMETRICSURFACE_API UTechSoftParametricSurfaceData : public UDatasmithParametricSurfaceData
 {
 	GENERATED_BODY()
 
diff --git a/Engine/Plugins/Enterprise/DatasmithCADImporter/Source/ParametricSurface/Public/DatasmithParametricSurfaceData.h b/Engine/Plugins/Enterprise/DatasmithCADImporter/Source/ParametricSurface/Public/DatasmithParametricSurfaceData.h
new file mode 100644
--- /dev/null
+++ b/Engine/Plugins/Enterprise/DatasmithCADImporter/Source/ParametricSurface/Public/DatasmithParametricSurfaceData.h
@@ -0,0 +1,120 @@
+// Copyright Epic Games, Inc. All Rights Reserved.
+
+#pragma once
+
+#include "CoreMinimal.h"
+
+#include "CADOptions.h"
+
+#include "DatasmithAdditionalData.h"
+#include "DatasmithCustomAction.h"
+#include "DatasmithImportOptions.h"
+#include "DatasmithUtils.h"
+
+#include "DatasmithParametricSurfaceData.generated.h"
+
+namespace CADLibrary
+{
+	class FImportParameters;
+	struct FMeshParameters;
+}
+
+USTRUCT(BlueprintType)
+struct PARAMETRICSURFACE_API FParametricSceneParameters
+{
+	GENERATED_BODY()
+
+	// value from FDatasmithUtils::EModelCoordSystem
+	UPROPERTY()
+	uint8 ModelCoordSys = (uint8)FDatasmithUtils::EModelCoordSystem::ZUp_LeftHanded;
+
+	UPROPERTY()
+	float MetricUnit = 0.01f;
+
+	UPROPERTY()
+	float ScaleFactor = 1.0f;
+};
+
+USTRUCT()
+struct PARAMETRICSURFACE_API FParametricMeshParameters
+{
+	GENERATED_BODY()
+
+	UPROPERTY()
+	bool bNeedSwapOrientation = false;
+
+	UPROPERTY()
+	bool bIsSymmetric = false;
+
+	UPROPERTY()
+	FVector SymmetricOrigin = FVector::ZeroVector;
+
+	UPROPERTY()
+	FVector SymmetricNormal = FVector::ZeroVector;
+
+	operator CADLibrary::FMeshParameters()
+	{
+		CADLibrary::FMeshParameters Parameters;
+
+		Parameters.bNeedSwapOrientation = bNeedSwapOrientation;
+		Parameters.bIsSymmetric = bIsSymmetric;
+		Parameters.SymmetricNormal = (FVector3f) SymmetricNormal;
+		Parameters.SymmetricOrigin = (FVector3f) SymmetricOrigin;
+
+		return Parameters;
+	}
+};
+
+UCLASS(meta = (DisplayName = "Datasmith Parametric Surface Data"))
+class PARAMETRICSURFACE_API UDatasmithParametricSurfaceData : public UDatasmithAdditionalData
+{
+	GENERATED_BODY()
+
+public:
+
+	virtual bool IsValid()
+	{
+		return RawData.Num() > 0;
+	}
+
+	virtual bool SetFile(const TCHAR* FilePath);
+
+	virtual void SetImportParameters(const CADLibrary::FImportParameters& InSceneParameters);
+
+	virtual void SetMeshParameters(const CADLibrary::FMeshParameters& InMeshParameters);
+
+	virtual const FDatasmithTessellationOptions& GetLastTessellationOptions() const
+	{
+		return LastTessellationOptions;
+	}
+
+	virtual void SetLastTessellationOptions(const FDatasmithTessellationOptions& InTessellationOptions)
+	{
+		LastTessellationOptions = InTessellationOptions;
+	}
+
+	virtual bool Tessellate(UStaticMesh& StaticMesh, const FDatasmithRetessellationOptions& RetessellateOptions)
+	{
+		return false;
+	}
+
+protected:
+	virtual void Serialize(FArchive& Ar) override;
+
+protected:
+	UPROPERTY()
+	FParametricSceneParameters SceneParameters;
+
+	UPROPERTY()
+	FParametricMeshParameters MeshParameters;
+
+	UPROPERTY(EditAnywhere, Category = NURBS)
+	FDatasmithTessellationOptions LastTessellationOptions;
+
+	UPROPERTY()
+	TArray<uint8> RawData_DEPRECATED;
+
+	// Too costly to serialize as a UPROPERTY, will use custom serialization.
+	TArray<uint8> RawData;
+};
+
diff --git a/Engine/Plugins/Enterprise/DatasmithCADImporter/Source/ParametricSurface/Public/ParametricSurfaceData.h b/Engine/Plugins/Enterprise/DatasmithCADImporter/Source/ParametricSurface/Public/ParametricSurfaceData.h
deleted file mode 100644
diff --git a/Engine/Plugins/Enterprise/DatasmithCADImporter/Source/ParametricSurface/Public/ParametricSurfaceModule.h b/Engine/Plugins/Enterprise/DatasmithCADImporter/Source/ParametricSurface/Public/ParametricSurfaceModule.h
--- a/Engine/Plugins/Enterprise/DatasmithCADImporter/Source/ParametricSurface/Public/ParametricSurfaceModule.h
+++ b/Engine/Plugins/Enterprise/DatasmithCADImporter/Source/ParametricSurface/Public/ParametricSurfaceModule.h
@@ -17,5 +17,5 @@ public:
 	static FParametricSurfaceModule& Get();
 	static bool IsAvailable();
 
-	static PARAMETRICSURFACE_API class UParametricSurfaceData* CreateParametricSurface();
+	static PARAMETRICSURFACE_API class UDatasmithParametricSurfaceData* CreateParametricSurface();
 };
diff --git a/Engine/Plugins/Enterprise/DatasmithCADImporter/Source/ParametricSurfaceExtension/ParametricSurfaceExtension.Build.cs b/Engine/Plugins/Enterprise/DatasmithCADImporter/Source/ParametricSurfaceExtension/ParametricSurfaceExtension.Build.cs
--- a/Engine/Plugins/Enterprise/DatasmithCADImporter/Source/ParametricSurfaceExtension/ParametricSurfaceExtension.Build.cs
+++ b/Engine/Plugins/Enterprise/DatasmithCADImporter/Source/ParametricSurfaceExtension/ParametricSurfaceExtension.Build.cs
@@ -18,24 +18,34 @@ namespace UnrealBuildTool.Rules
 			PrivateDependencyModuleNames.AddRange(
 				new string[]
 				{
-					"Core",
 					"CoreUObject",
-					"MeshDescription",
-					"DataprepCore",
+					"CADLibrary",
 					"DatasmithContent",
+					"Engine",
+					"MeshDescription",
 					"ParametricSurface",
 					"PhysicsCore",
-					"DatasmithImporter",
-					"Engine",
-					"StaticMeshEditor",
-					"Slate",
-					"SlateCore",
 					"StaticMeshDescription",
-					"EditorFramework",
-					"UnrealEd",
-					"CADLibrary",
 				}
 			);
+
+			if (Target.bBuildEditor)
+			{
+				PrivateDependencyModuleNames.AddRange(
+					new string[]
+					{
+						"ContentBrowser",
+						"DataprepCore",
+						"DatasmithImporter",
+						"EditorFramework",
+						"Slate",
+						"SlateCore",
+						"StaticMeshEditor",
+						"ToolMenus",
+						"UnrealEd",
+					}
+				);
+			}
 		}
 	}
 }
diff --git a/Engine/Plugins/Enterprise/DatasmithCADImporter/Source/ParametricSurfaceExtension/Private/DataprepTessellationOperation.h b/Engine/Plugins/Enterprise/DatasmithCADImporter/Source/ParametricSurfaceExtension/Private/DataprepTessellationOperation.h
--- a/Engine/Plugins/Enterprise/DatasmithCADImporter/Source/ParametricSurfaceExtension/Private/DataprepTessellationOperation.h
+++ b/Engine/Plugins/Enterprise/DatasmithCADImporter/Source/ParametricSurfaceExtension/Private/DataprepTessellationOperation.h
@@ -15,7 +15,7 @@ namespace CADLibrary
 	class FImportParameters;
 }
 struct FMeshDescription;
-class UParametricSurfaceData;
+class UDatasmithParametricSurfaceData;
 class UDatasmithStaticMeshCADImportData;
 
 UCLASS(Experimental, Category = MeshOperation, Meta = (DisplayName="Datasmith Tessellation", ToolTip = "For each static mesh to process, retessellate the mesh if the object contains the required data") )
diff --git a/Engine/Plugins/Enterprise/DatasmithCADImporter/Source/ParametricSurfaceExtension/Private/ParametricRetessellateAction.cpp b/Engine/Plugins/Enterprise/DatasmithCADImporter/Source/ParametricSurfaceExtension/Private/ParametricRetessellateAction.cpp
--- a/Engine/Plugins/Enterprise/DatasmithCADImporter/Source/ParametricSurfaceExtension/Private/ParametricRetessellateAction.cpp
+++ b/Engine/Plugins/Enterprise/DatasmithCADImporter/Source/ParametricSurfaceExtension/Private/ParametricRetessellateAction.cpp
@@ -2,32 +2,24 @@
 
 #include "ParametricRetessellateAction.h"
 
-#include "Components/StaticMeshComponent.h"
 #include "ParametricRetessellateAction_Impl.h"
 
 #include "DatasmithAdditionalData.h"
+#include "DatasmithParametricSurfaceData.h"
 #include "DatasmithStaticMeshImporter.h" // Call to BuildStaticMesh
-#include "DatasmithUtils.h"
 #include "DatasmithTranslator.h"
-#include "Physics/Experimental/PhysScene_Chaos.h"
-#include "Physics/PhysScene.h"
 #include "UI/DatasmithDisplayHelper.h"
-#include "MeshDescriptionHelper.h"
 
 #include "Algo/AnyOf.h"
-#include "Algo/Transform.h"
-#include "AssetRegistry/AssetData.h"
 #include "Async/ParallelFor.h"
-#include "Chaos/ChaosScene.h"
+#include "Components/StaticMeshComponent.h"
 #include "Engine/StaticMesh.h"
-#include "HAL/PlatformFileManager.h"
+#include "Engine/World.h"
 #include "IStaticMeshEditor.h"
-#include "Misc/FileHelper.h"
 #include "Misc/ScopedSlowTask.h"
-#include "StaticMeshAttributes.h"
+#include "Physics/Experimental/PhysScene_Chaos.h"
 #include "Toolkits/IToolkit.h"
 #include "Toolkits/ToolkitManager.h"
-#include "UObject/StrongObjectPtr.h"
 
 
 #define LOCTEXT_NAMESPACE "ParametricRetessellateAction"
@@ -57,7 +49,7 @@ void UParametricRetessellateAction::ApplyOnAssets(const TArray<FAssetData>& Sele
 
 bool FParametricRetessellateAction_Impl::CanApplyOnAssets(const TArray<FAssetData>& SelectedAssets)
 {
-	return Algo::AnyOf(SelectedAssets, [](const FAssetData& Asset) { return Datasmith::GetAdditionalData<UParametricSurfaceData>(Asset) != nullptr; });
+	return Algo::AnyOf(SelectedAssets, [](const FAssetData& Asset) { return Datasmith::GetAdditionalData<UDatasmithParametricSurfaceData>(Asset) != nullptr; });
 }
 
 void FParametricRetessellateAction_Impl::ApplyOnAssets(const TArray<FAssetData>& SelectedAssets)
@@ -89,7 +81,7 @@ void FParametricRetessellateAction_Impl::ApplyOnAssets(const TArray<FAssetData>&
 	for (const FAssetData& Asset : SelectedAssets)
 	{
 		AssetIndex++;
-		if (UParametricSurfaceData* ParametricSurfaceData = Datasmith::GetAdditionalData<UParametricSurfaceData>(Asset))
+		if (UDatasmithParametricSurfaceData* ParametricSurfaceData = Datasmith::GetAdditionalData<UDatasmithParametricSurfaceData>(Asset))
 		{
 			if (ParametricSurfaceData->IsValid())
 			{
@@ -212,7 +204,7 @@ TSet<UStaticMesh*> GetReferencedStaticMeshes(const TArray<AActor*>& SelectedActo
 bool UParametricRetessellateAction::CanApplyOnActors(const TArray<AActor*>& SelectedActors)
 {
 	const TSet<UStaticMesh*> ReferencedStaticMeshes = GetReferencedStaticMeshes(SelectedActors);
-	return Algo::AnyOf(ReferencedStaticMeshes, [](const UStaticMesh* Mesh) { return Datasmith::GetAdditionalData<UParametricSurfaceData>(FAssetData(Mesh)); });
+	return Algo::AnyOf(ReferencedStaticMeshes, [](const UStaticMesh* Mesh) { return Datasmith::GetAdditionalData<UDatasmithParametricSurfaceData>(FAssetData(Mesh)); });
 }
 
 void UParametricRetessellateAction::ApplyOnActors(const TArray<AActor*>& SelectedActors)
diff --git a/Engine/Plugins/Enterprise/DatasmithCADImporter/Source/ParametricSurfaceExtension/Private/ParametricSurfaceBlueprintLibrary.cpp b/Engine/Plugins/Enterprise/DatasmithCADImporter/Source/ParametricSurfaceExtension/Private/ParametricSurfaceBlueprintLibrary.cpp
--- a/Engine/Plugins/Enterprise/DatasmithCADImporter/Source/ParametricSurfaceExtension/Private/ParametricSurfaceBlueprintLibrary.cpp
+++ b/Engine/Plugins/Enterprise/DatasmithCADImporter/Source/ParametricSurfaceExtension/Private/ParametricSurfaceBlueprintLibrary.cpp
@@ -3,7 +3,7 @@
 #include "ParametricSurfaceBlueprintLibrary.h"
 
 #include "ParametricRetessellateAction.h"
-#include "DatasmithAdditionalData.h"
+#include "DatasmithParametricSurfaceData.h"
 #include "DatasmithStaticMeshImporter.h" // Call to BuildStaticMesh
 #include "DatasmithUtils.h"
 
@@ -33,7 +33,7 @@ bool UParametricSurfaceBlueprintLibrary::RetessellateStaticMeshWithNotification(
 	int32 LODIndex = 0;
 
 	FAssetData AssetData( StaticMesh );
-	if (UParametricSurfaceData* ParametricSurfaceData = Datasmith::GetAdditionalData<UParametricSurfaceData>(AssetData))
+	if (UDatasmithParametricSurfaceData* DatasmithParametricSurfaceData = Datasmith::GetAdditionalData<UDatasmithParametricSurfaceData>(AssetData))
 	{
 		// Make sure MeshDescription exists
 		FMeshDescription* DestinationMeshDescription = StaticMesh->GetMeshDescription( LODIndex );
@@ -51,7 +51,7 @@ bool UParametricSurfaceBlueprintLibrary::RetessellateStaticMeshWithNotification(
 			}
 
 			const int32 OldNumberOfUVChannels = FStaticMeshAttributes(*DestinationMeshDescription).GetVertexInstanceUVs().GetNumChannels();
-			if (ParametricSurfaceData->Tessellate(*StaticMesh, TessellationSettings))
+			if (DatasmithParametricSurfaceData->Tessellate(*StaticMesh, TessellationSettings))
 			{
 				const int32 NumberOfUVChannels = FStaticMeshAttributes(*DestinationMeshDescription).GetVertexInstanceUVs().GetNumChannels();
 				if (NumberOfUVChannels < OldNumberOfUVChannels)
@@ -105,7 +105,7 @@ bool UParametricSurfaceBlueprintLibrary::RetessellateStaticMeshWithNotification(
 				}
 
 				// Save last tessellation settings
-				ParametricSurfaceData->SetLastTessellationOptions(TessellationSettings);
+				DatasmithParametricSurfaceData->SetLastTessellationOptions(TessellationSettings);
 
 				bTessellationOutcome = true;
 			}
diff --git a/Engine/Plugins/Enterprise/DatasmithCADImporter/Source/ParametricSurfaceExtension/Private/ParametricSurfaceExtensionModule.cpp b/Engine/Plugins/Enterprise/DatasmithCADImporter/Source/ParametricSurfaceExtension/Private/ParametricSurfaceExtensionModule.cpp
--- a/Engine/Plugins/Enterprise/DatasmithCADImporter/Source/ParametricSurfaceExtension/Private/ParametricSurfaceExtensionModule.cpp
+++ b/Engine/Plugins/Enterprise/DatasmithCADImporter/Source/ParametricSurfaceExtension/Private/ParametricSurfaceExtensionModule.cpp
@@ -5,11 +5,14 @@
 #include "ParametricRetessellateAction.h"
 #include "ParametricRetessellateAction_Impl.h"
 
+#include "ContentBrowserMenuContexts.h"
 #include "Framework/MultiBox/MultiBoxBuilder.h"
 #include "IStaticMeshEditor.h"
 #include "Modules/ModuleInterface.h"
 #include "Modules/ModuleManager.h"
 #include "StaticMeshEditorModule.h"
+#include "Styling/AppStyle.h"
+#include "ToolMenu.h"
 #include "UObject/StrongObjectPtr.h"
 
 
@@ -67,14 +70,50 @@ namespace StaticMeshEditorExtenser
 		return Extender;
 	}
 
+	void ExecuteRetessellation(const FToolMenuContext& InContext)
+	{
+		if (const UContentBrowserAssetContextMenuContext* CBContext = InContext.FindContext<UContentBrowserAssetContextMenuContext>())
+		{
+			FParametricRetessellateAction_Impl::ApplyOnAssets(CBContext->SelectedAssets);
+		}
+	}
+
+	bool CanExecuteRetessellation(const FToolMenuContext& InContext)
+	{
+		if (const UContentBrowserAssetContextMenuContext* CBContext = InContext.FindContext<UContentBrowserAssetContextMenuContext>())
+		{
+			return FParametricRetessellateAction_Impl::CanApplyOnAssets(CBContext->SelectedAssets);
+		}
+
+		return false;
+	}
+
 	void Register()
 	{
-		if (!IsRunningCommandlet())
+		if (IsRunningCommandlet())
 		{
-			IStaticMeshEditorModule& StaticMeshEditorModule = FModuleManager::Get().LoadModuleChecked<IStaticMeshEditorModule>("StaticMeshEditor");
-			TArray<FAssetEditorExtender>& ExtenderDelegates = StaticMeshEditorModule.GetMenuExtensibilityManager()->GetExtenderDelegates();
-			ExtenderDelegates.Add(FAssetEditorExtender::CreateStatic(&StaticMeshEditorExtenser::CreateExtenderForObjects));
+			return;
 		}
+
+		IStaticMeshEditorModule& StaticMeshEditorModule = FModuleManager::Get().LoadModuleChecked<IStaticMeshEditorModule>("StaticMeshEditor");
+		TArray<FAssetEditorExtender>& ExtenderDelegates = StaticMeshEditorModule.GetMenuExtensibilityManager()->GetExtenderDelegates();
+		ExtenderDelegates.Add(FAssetEditorExtender::CreateStatic(&StaticMeshEditorExtenser::CreateExtenderForObjects));
+
+		UToolMenu* Menu = UE::ContentBrowser::ExtendToolMenu_AssetContextMenu(UStaticMesh::StaticClass());
+
+		FToolMenuSection& Section = Menu->FindOrAddSection("GetAssetActions");
+		Section.AddDynamicEntry(NAME_None, FNewToolMenuSectionDelegate::CreateLambda([](FToolMenuSection& InSection)
+			{
+				{
+					const TAttribute<FText> Label = LOCTEXT("Retessallation_Execute", "Retessellate");
+					const TAttribute<FText> ToolTip = LOCTEXT("Retessalltion_Execute_Tooltip", "Retessellate.");
+					const FSlateIcon SlateIcon(FAppStyle::GetAppStyleSetName(), "Icons.Adjust");
+					FToolUIAction UIAction;
+					UIAction.ExecuteAction = FToolMenuExecuteAction::CreateStatic(&ExecuteRetessellation);
+					UIAction.IsActionVisibleDelegate = FToolMenuCanExecuteAction::CreateStatic(&CanExecuteRetessellation);
+					InSection.AddMenuEntry("FixTextureValidation", Label, ToolTip, SlateIcon, UIAction);
+				}
+			}));
 	}
 };
 
diff --git a/Engine/Plugins/Enterprise/DatasmithCADImporter/Source/ParametricSurfaceExtension/Public/ParametricRetessellateAction.h b/Engine/Plugins/Enterprise/DatasmithCADImporter/Source/ParametricSurfaceExtension/Public/ParametricRetessellateAction.h
--- a/Engine/Plugins/Enterprise/DatasmithCADImporter/Source/ParametricSurfaceExtension/Public/ParametricRetessellateAction.h
+++ b/Engine/Plugins/Enterprise/DatasmithCADImporter/Source/ParametricSurfaceExtension/Public/ParametricRetessellateAction.h
@@ -4,9 +4,7 @@
 
 #include "CoreMinimal.h"
 
-#include "ParametricSurfaceData.h"
 #include "DatasmithCustomAction.h"
-#include "DatasmithAdditionalData.h"
 #include "DatasmithImportOptions.h"
 
 #include "ParametricRetessellateAction.generated.h"
diff --git a/Engine/Plugins/Enterprise/DatasmithContent/Config/BaseDatasmithContent.ini b/Engine/Plugins/Enterprise/DatasmithContent/Config/BaseDatasmithContent.ini
--- a/Engine/Plugins/Enterprise/DatasmithContent/Config/BaseDatasmithContent.ini
+++ b/Engine/Plugins/Enterprise/DatasmithContent/Config/BaseDatasmithContent.ini
@@ -17,3 +17,7 @@
 +PackageRedirects=(OldName="/DatasmithContent/Materials/FBXImporter/VREDMasterAbsorb.VREDMasterAbsorb",NewName="/DatasmithContent/Materials/FBXImporter/VREDReferenceAbsorb.VREDReferenceAbsorb")
 +PackageRedirects=(OldName="/DatasmithContent/Materials/FBXImporter/VREDMasterClearCoat.VREDMasterClearCoat",NewName="/DatasmithContent/Materials/FBXImporter/VREDReferenceClearCoat.VREDReferenceClearCoat")
 +PackageRedirects=(OldName="/DatasmithContent/Materials/FBXImporter/VREDMasterTransparent.VREDMasterTransparent",NewName="/DatasmithContent/Materials/FBXImporter/VREDReferenceTransparent.VREDReferenceTransparent")
+
+[SectionsToSave]
++Section=/Script/DatasmithContent.DatasmithCommonTessellationOptions
++Section=/Script/DatasmithContent.DatasmithImportOptions
diff --git a/Engine/Plugins/Experimental/Enterprise/DatasmithInterchange/Source/DatasmithInterchange/DatasmithInterchange.Build.cs b/Engine/Plugins/Experimental/Enterprise/DatasmithInterchange/Source/DatasmithInterchange/DatasmithInterchange.Build.cs
--- a/Engine/Plugins/Experimental/Enterprise/DatasmithInterchange/Source/DatasmithInterchange/DatasmithInterchange.Build.cs
+++ b/Engine/Plugins/Experimental/Enterprise/DatasmithInterchange/Source/DatasmithInterchange/DatasmithInterchange.Build.cs
@@ -23,6 +23,7 @@ namespace UnrealBuildTool.Rules
 					"InterchangeImport",
 					"InterchangePipelines",
 					"MeshDescription",
+					"ParametricSurface",
 					"StaticMeshDescription"
 				}
 			);
@@ -36,6 +37,7 @@ namespace UnrealBuildTool.Rules
 						"DesktopPlatform",
 						"InputCore",
 						"InterchangeEditorPipelines",
+						"PropertyEditor",
 						"Slate",
 						"SlateCore",
 						"ToolMenus",
diff --git a/Engine/Plugins/Experimental/Enterprise/DatasmithInterchange/Source/DatasmithInterchange/Private/DatasmithInterchangeModule.cpp b/Engine/Plugins/Experimental/Enterprise/DatasmithInterchange/Source/DatasmithInterchange/Private/DatasmithInterchangeModule.cpp
--- a/Engine/Plugins/Experimental/Enterprise/DatasmithInterchange/Source/DatasmithInterchange/Private/DatasmithInterchangeModule.cpp
+++ b/Engine/Plugins/Experimental/Enterprise/DatasmithInterchange/Source/DatasmithInterchange/Private/DatasmithInterchangeModule.cpp
@@ -4,6 +4,7 @@
 
 #include "Engine/Blueprint.h"
 #include "InterchangeDatasmithAreaLightFactory.h"
+#include "InterchangeDatasmithCustomizations.h"
 #include "InterchangeDatasmithLog.h"
 #include "InterchangeDatasmithPipeline.h"
 #include "InterchangeDatasmithTranslator.h"
@@ -21,9 +22,14 @@
 #include "InterchangeGenericScenesPipeline.h"
 #include "InterchangeManager.h"
 #include "InterchangeProjectSettings.h"
+
 #include "Logging/LogMacros.h"
 #include "UObject/SoftObjectPath.h"
 
+#if WITH_EDITOR
+#include "PropertyEditorModule.h"
+#endif
+
 #define LOCTEXT_NAMESPACE "DatasmithInterchange"
 
 
@@ -68,12 +74,37 @@ public:
 		FDatasmithReferenceMaterialManager::Get().RegisterSelector(TEXT("SketchUp"), MakeShared< FDatasmithSketchUpMaterialSelector >());
 		FDatasmithReferenceMaterialManager::Get().RegisterSelector(TEXT("CityEngine"), MakeShared< FDatasmithCityEngineMaterialSelector >());
 		FDatasmithReferenceMaterialManager::Get().RegisterSelector(TEXT("StdMaterial"), MakeShared< FDatasmithStdMaterialSelector >());
+		
+#if WITH_EDITOR
+		ClassesToUnregisterOnShutdown.Reset();
+		// Register details customizations
+		FPropertyEditorModule& PropertyEditorModule = FModuleManager::LoadModuleChecked<FPropertyEditorModule>("PropertyEditor");
+
+		ClassesToUnregisterOnShutdown.Add(UInterchangeDatasmithTranslatorSettings::StaticClass()->GetFName());
+		PropertyEditorModule.RegisterCustomClassLayout(ClassesToUnregisterOnShutdown.Last(), FOnGetDetailCustomizationInstance::CreateStatic(&FInterchangeDatasmithTranslatorSettingsCustomization::MakeInstance));
+#endif
 	}
 
 	virtual void ShutdownModule() override
 	{
 		UE::DatasmithInterchange::FDatasmithReferenceMaterialManager::Destroy();
+		
+#if WITH_EDITOR
+		FPropertyEditorModule* PropertyEditorModule = FModuleManager::GetModulePtr<FPropertyEditorModule>("PropertyEditor");
+		if (PropertyEditorModule)
+		{
+			for (FName ClassName : ClassesToUnregisterOnShutdown)
+			{
+				PropertyEditorModule->UnregisterCustomClassLayout(ClassName);
+			}
+		}
+		ClassesToUnregisterOnShutdown.Empty();
+#endif
 	}
+
+#if WITH_EDITOR
+	TArray<FName> ClassesToUnregisterOnShutdown;
+#endif
 };
 
 IMPLEMENT_MODULE(FDatasmithInterchangeModule, DatasmithInterchange);
diff --git a/Engine/Plugins/Experimental/Enterprise/DatasmithInterchange/Source/DatasmithInterchange/Private/InterchangeDatasmithCustomizations.cpp b/Engine/Plugins/Experimental/Enterprise/DatasmithInterchange/Source/DatasmithInterchange/Private/InterchangeDatasmithCustomizations.cpp
new file mode 100644
--- /dev/null
+++ b/Engine/Plugins/Experimental/Enterprise/DatasmithInterchange/Source/DatasmithInterchange/Private/InterchangeDatasmithCustomizations.cpp
@@ -0,0 +1,56 @@
+// Copyright Epic Games, Inc. All Rights Reserved.
+
+#include "InterchangeDatasmithCustomizations.h"
+
+#if WITH_EDITOR
+#include "InterchangeDatasmithTranslator.h"
+
+#include "DetailLayoutBuilder.h"
+#include "Modules/ModuleManager.h"
+#include "PropertyCustomizationHelpers.h"
+#include "UObject/UnrealType.h"
+#include "Widgets/Text/STextBlock.h"
+
+TSharedRef<IDetailCustomization> FInterchangeDatasmithTranslatorSettingsCustomization::MakeInstance()
+{
+	return MakeShareable(new FInterchangeDatasmithTranslatorSettingsCustomization());
+}
+
+void FInterchangeDatasmithTranslatorSettingsCustomization::CustomizeDetails(IDetailLayoutBuilder& DetailBuilder)
+{
+	TArray<TWeakObjectPtr<UObject>> EditingObjects;
+	DetailBuilder.GetObjectsBeingCustomized(EditingObjects);
+	ensure(EditingObjects.Num() == 1);
+
+	UInterchangeDatasmithTranslatorSettings* TranslatorSettings = Cast<UInterchangeDatasmithTranslatorSettings>(EditingObjects[0].Get());
+
+	if (!ensure(TranslatorSettings && TranslatorSettings->DatasmithOption))
+	{
+		return;
+	}
+
+	TSharedRef<IPropertyHandle> DatasmithOptionHandle = DetailBuilder.GetProperty(TEXT("DatasmithOption"), UInterchangeDatasmithTranslatorSettings::StaticClass());
+
+	FPropertyEditorModule& PropertyEditorModule = FModuleManager::GetModuleChecked<FPropertyEditorModule>("PropertyEditor");
+
+	FDetailsViewArgs DetailsViewArgs;
+	DetailsViewArgs.NameAreaSettings = FDetailsViewArgs::HideNameArea;
+	DetailsViewArgs.bAllowSearch = false;
+	TSharedPtr<IDetailsView> DetailsView = PropertyEditorModule.CreateDetailView(DetailsViewArgs);
+
+	DetailsView->OnFinishedChangingProperties().AddLambda([DatasmithOptionHandle](const FPropertyChangedEvent& InPropertyChangedEvent)
+		{
+			DatasmithOptionHandle->NotifyFinishedChangingProperties();
+		});
+
+	DetailsView->SetObject(TranslatorSettings->DatasmithOption);
+
+	IDetailPropertyRow* DetailPropertyRow = DetailBuilder.EditDefaultProperty(DatasmithOptionHandle);
+
+	DetailPropertyRow->CustomWidget()
+	[
+		DetailsView.ToSharedRef()
+	];
+}
+
+#endif
\ No newline at end of file
diff --git a/Engine/Plugins/Experimental/Enterprise/DatasmithInterchange/Source/DatasmithInterchange/Private/InterchangeDatasmithCustomizations.h b/Engine/Plugins/Experimental/Enterprise/DatasmithInterchange/Source/DatasmithInterchange/Private/InterchangeDatasmithCustomizations.h
new file mode 100644
--- /dev/null
+++ b/Engine/Plugins/Experimental/Enterprise/DatasmithInterchange/Source/DatasmithInterchange/Private/InterchangeDatasmithCustomizations.h
@@ -0,0 +1,26 @@
+// Copyright Epic Games, Inc. All Rights Reserved.
+
+#pragma once
+
+#if WITH_EDITOR
+#include "IDetailCustomization.h"
+#include "UObject/WeakObjectPtr.h"
+
+class IDetailLayoutBuilder;
+class UInterchangeDatasmithTranslatorSettings;
+struct FAssetData;
+
+class FInterchangeDatasmithTranslatorSettingsCustomization : public IDetailCustomization
+{
+public:
+
+	/**
+	 * Creates an instance of this class
+	 */
+	static TSharedRef<IDetailCustomization> MakeInstance();
+
+	/** IDetailCustomization interface */
+	virtual void CustomizeDetails(IDetailLayoutBuilder& DetailBuilder) override;
+	/** End IDetailCustomization interface */
+};
+#endif
diff --git a/Engine/Plugins/Experimental/Enterprise/DatasmithInterchange/Source/DatasmithInterchange/Private/InterchangeDatasmithPipeline.cpp b/Engine/Plugins/Experimental/Enterprise/DatasmithInterchange/Source/DatasmithInterchange/Private/InterchangeDatasmithPipeline.cpp
--- a/Engine/Plugins/Experimental/Enterprise/DatasmithInterchange/Source/DatasmithInterchange/Private/InterchangeDatasmithPipeline.cpp
+++ b/Engine/Plugins/Experimental/Enterprise/DatasmithInterchange/Source/DatasmithInterchange/Private/InterchangeDatasmithPipeline.cpp
@@ -33,14 +33,6 @@
 #include "Engine/StaticMeshActor.h"
 #include "Components/StaticMeshComponent.h"
 
-#if WITH_EDITOR
-#include "DatasmithImporter.h"
-#include "DatasmithImportContext.h"
-#include "DatasmithStaticMeshImporter.h"
-#include "Misc/ScopedSlowTask.h"
-#include "Utility/DatasmithImporterUtils.h"
-#endif //WITH_EDITOR
-
 #define LOCTEXT_NAMESPACE "InterchangeDatasmithPipeline"
 
 namespace UE::Interchange::StaticMeshUtils
diff --git a/Engine/Plugins/Experimental/Enterprise/DatasmithInterchange/Source/DatasmithInterchange/Private/InterchangeDatasmithStaticMeshData.cpp b/Engine/Plugins/Experimental/Enterprise/DatasmithInterchange/Source/DatasmithInterchange/Private/InterchangeDatasmithStaticMeshData.cpp
new file mode 100644
--- /dev/null
+++ b/Engine/Plugins/Experimental/Enterprise/DatasmithInterchange/Source/DatasmithInterchange/Private/InterchangeDatasmithStaticMeshData.cpp
@@ -0,0 +1,4 @@
+// Copyright Epic Games, Inc. All Rights Reserved.
+
+#include "InterchangeDatasmithStaticMeshData.h"
+
diff --git a/Engine/Plugins/Experimental/Enterprise/DatasmithInterchange/Source/DatasmithInterchange/Private/InterchangeDatasmithStaticMeshData.h b/Engine/Plugins/Experimental/Enterprise/DatasmithInterchange/Source/DatasmithInterchange/Private/InterchangeDatasmithStaticMeshData.h
new file mode 100644
--- /dev/null
+++ b/Engine/Plugins/Experimental/Enterprise/DatasmithInterchange/Source/DatasmithInterchange/Private/InterchangeDatasmithStaticMeshData.h
@@ -0,0 +1,23 @@
+// Copyright Epic Games, Inc. All Rights Reserved.
+
+#pragma once
+
+#include "CoreMinimal.h"
+
+#include "DatasmithAdditionalData.h"
+#include "DatasmithPayload.h"
+#include "Nodes/InterchangeBaseNode.h"
+
+#include "InterchangeDatasmithStaticMeshData.generated.h"
+
+UCLASS()
+class UDatasmithInterchangeStaticMeshDataNode : public UInterchangeBaseNode
+{
+	GENERATED_BODY()
+
+public:
+	/**
+	*/
+	UPROPERTY()
+	TMap<FString, TObjectPtr<UDatasmithAdditionalData>> AdditionalDataMap;
+};
diff --git a/Engine/Plugins/Experimental/Enterprise/DatasmithInterchange/Source/DatasmithInterchange/Private/InterchangeDatasmithStaticMeshPipeline.cpp b/Engine/Plugins/Experimental/Enterprise/DatasmithInterchange/Source/DatasmithInterchange/Private/InterchangeDatasmithStaticMeshPipeline.cpp
--- a/Engine/Plugins/Experimental/Enterprise/DatasmithInterchange/Source/DatasmithInterchange/Private/InterchangeDatasmithStaticMeshPipeline.cpp
+++ b/Engine/Plugins/Experimental/Enterprise/DatasmithInterchange/Source/DatasmithInterchange/Private/InterchangeDatasmithStaticMeshPipeline.cpp
@@ -3,13 +3,18 @@
 #include "InterchangeDatasmithStaticMeshPipeline.h"
 
 #include "InterchangeDatasmithUtils.h"
+#include "InterchangeDatasmithStaticMeshData.h"
 
 #include "InterchangeMaterialFactoryNode.h"
 #include "InterchangeStaticMeshFactoryNode.h"
 #include "InterchangeMeshNode.h"
 
+#include "DatasmithAssetImportData.h"
+#include "DatasmithImportOptions.h"
+
 #include "Engine/StaticMesh.h"
 #include "Materials/MaterialInterface.h"
+#include "DatasmithParametricSurfaceData.h"
 
 void UInterchangeDatasmithStaticMeshPipeline::ExecutePipeline(UInterchangeBaseNodeContainer* NodeContainer, const TArray<UInterchangeSourceData*>& InSourceDatas, const FString& ContentBasePath)
 {
@@ -69,6 +74,14 @@ void UInterchangeDatasmithStaticMeshPipeline::ExecutePostImportPipeline(const UI
 		return;
 	}
 
+	ApplyMaterials(NodeContainer, FactoryNode, StaticMesh);
+	ApplyAdditionalData(NodeContainer, FactoryNode, StaticMesh);
+}
+
+void UInterchangeDatasmithStaticMeshPipeline::ApplyMaterials(const UInterchangeBaseNodeContainer* NodeContainer, const UInterchangeStaticMeshFactoryNode* FactoryNode, UStaticMesh* StaticMesh) const
+{
+	using namespace UE::DatasmithInterchange;
+
 	FString MaterialFactoryUid;
 	if (!FactoryNode->GetStringAttribute(MeshUtils::MeshMaterialAttrName, MaterialFactoryUid))
 	{
@@ -92,3 +105,68 @@ void UInterchangeDatasmithStaticMeshPipeline::ExecutePostImportPipeline(const UI
 		}
 	}
 }
+
+void UInterchangeDatasmithStaticMeshPipeline::ApplyAdditionalData(const UInterchangeBaseNodeContainer* NodeContainer, const UInterchangeStaticMeshFactoryNode* FactoryNode, UStaticMesh* StaticMesh) const
+{
+#if WITH_EDITORONLY_DATA
+	TArray<FString> TargetNodes;
+	FactoryNode->GetTargetNodeUids(TargetNodes);
+	if (TargetNodes.Num() == 0)
+	{
+		return;
+	}
+
+	const FString& MeshNodeUid = TargetNodes[0];
+
+	UDatasmithInterchangeStaticMeshDataNode* StaticMeshDataNode = nullptr;
+	NodeContainer->IterateNodes([&StaticMeshDataNode](const FString& NodeUid, UInterchangeBaseNode* Node)
+		{
+			if (UDatasmithInterchangeStaticMeshDataNode* TypedNode = Cast<UDatasmithInterchangeStaticMeshDataNode>(Node))
+			{
+				StaticMeshDataNode = TypedNode;
+				return;
+			}
+		});
+
+	if (!StaticMeshDataNode)
+	{
+		return;
+	}
+
+	if (const TObjectPtr<UDatasmithAdditionalData>* AdditionalDataPtr = StaticMeshDataNode->AdditionalDataMap.Find(MeshNodeUid))
+	{
+		if (UDatasmithParametricSurfaceData* AdditionalData = Cast<UDatasmithParametricSurfaceData>(AdditionalDataPtr->Get()))
+		{
+			FDatasmithStaticMeshImportOptions StaticMeshOptions;
+			FDatasmithAssetImportOptions AssetOptions;
+			
+			UDatasmithStaticMeshImportData::DefaultOptionsPair ImportOptions = UDatasmithStaticMeshImportData::DefaultOptionsPair(StaticMeshOptions, AssetOptions);
+
+			UDatasmithStaticMeshImportData* MeshImportData = UDatasmithStaticMeshImportData::GetImportDataForStaticMesh(StaticMesh, ImportOptions);
+
+			if (MeshImportData)
+			{
+				FString SourceUri;
+				FMD5Hash SourceHash;
+				if (SourceDatas.Num() > 0)
+				{
+					SourceUri = SourceDatas[0]->GetFilename();
+					TOptional<FMD5Hash> OptHash = SourceDatas[0]->GetFileContentHash();
+					if (OptHash.IsSet())
+					{
+						SourceHash = OptHash.GetValue();
+					}
+				}
+				// Update the import data source file and set the mesh hash
+				// #ueent_todo FH: piggybacking off of the SourceData file hash for now, until we have custom derived AssetImportData properly serialize to the AssetRegistry
+				MeshImportData->Update(SourceUri, &SourceHash);
+				MeshImportData->DatasmithImportInfo = FDatasmithImportInfo(SourceUri, SourceHash);
+
+				// Set the final outer // #ueent_review: propagate flags of outer?
+				AdditionalData->Rename(nullptr, MeshImportData);
+				MeshImportData->AdditionalData.Add(AdditionalData);
+			}
+		}
+	}
+#endif
+}
diff --git a/Engine/Plugins/Experimental/Enterprise/DatasmithInterchange/Source/DatasmithInterchange/Private/InterchangeDatasmithStaticMeshPipeline.h b/Engine/Plugins/Experimental/Enterprise/DatasmithInterchange/Source/DatasmithInterchange/Private/InterchangeDatasmithStaticMeshPipeline.h
--- a/Engine/Plugins/Experimental/Enterprise/DatasmithInterchange/Source/DatasmithInterchange/Private/InterchangeDatasmithStaticMeshPipeline.h
+++ b/Engine/Plugins/Experimental/Enterprise/DatasmithInterchange/Source/DatasmithInterchange/Private/InterchangeDatasmithStaticMeshPipeline.h
@@ -25,4 +25,8 @@ protected:
 		// This pipeline creates UObjects and assets. Not safe to execute outside of main thread.
 		return false;
 	}
+
+private:
+	void ApplyMaterials(const UInterchangeBaseNodeContainer* NodeContainer, const UInterchangeStaticMeshFactoryNode* FactoryNode, UStaticMesh* StaticMesh) const;
+	void ApplyAdditionalData(const UInterchangeBaseNodeContainer* NodeContainer, const UInterchangeStaticMeshFactoryNode* FactoryNode, UStaticMesh* StaticMesh) const;
 };
\ No newline at end of file
diff --git a/Engine/Plugins/Experimental/Enterprise/DatasmithInterchange/Source/DatasmithInterchange/Private/InterchangeDatasmithTranslator.cpp b/Engine/Plugins/Experimental/Enterprise/DatasmithInterchange/Source/DatasmithInterchange/Private/InterchangeDatasmithTranslator.cpp
--- a/Engine/Plugins/Experimental/Enterprise/DatasmithInterchange/Source/DatasmithInterchange/Private/InterchangeDatasmithTranslator.cpp
+++ b/Engine/Plugins/Experimental/Enterprise/DatasmithInterchange/Source/DatasmithInterchange/Private/InterchangeDatasmithTranslator.cpp
@@ -5,6 +5,7 @@
 #include "InterchangeDatasmithAreaLightNode.h"
 #include "InterchangeDatasmithLog.h"
 #include "InterchangeDatasmithMaterialNode.h"
+#include "InterchangeDatasmithStaticMeshData.h"
 #include "InterchangeDatasmithTextureData.h"
 #include "InterchangeDatasmithUtils.h"
 
@@ -17,6 +18,7 @@
 #include "DatasmithUtils.h"
 #include "DatasmithVariantElements.h"
 #include "IDatasmithSceneElements.h"
+#include "DatasmithParametricSurfaceData.h"
 
 #include "CADOptions.h"
 #include "ExternalSourceModule.h"
@@ -37,6 +39,7 @@
 #include "InterchangeVariantSetNode.h"
 #include "StaticMeshOperations.h"
 
+#include "HAL/ConsoleManager.h"
 #include "Misc/App.h"
 #include "Misc/PackageName.h"
 
@@ -52,70 +55,6 @@
 
 #define LOCTEXT_NAMESPACE "DatasmithInterchange"
 
-namespace UE::DatasmithInterchange
-{
-#if WITH_EDITOR
-	bool DisplayOptionsDialog(IDatasmithTranslator& Translator)
-	{
-		TArray<TObjectPtr<UDatasmithOptionsBase>> ImportOptions;
-		Translator.GetSceneImportOptions(ImportOptions);
-
-		if (ImportOptions.Num() == 0)
-		{
-			return true;
-		}
-
-		const FString FilePath = Translator.GetSource().GetSourceFile();
-
-		TSharedPtr<SWindow> ParentWindow;
-
-		if (FModuleManager::Get().IsModuleLoaded("MainFrame"))
-		{
-			IMainFrameModule& MainFrame = FModuleManager::LoadModuleChecked<IMainFrameModule>("MainFrame");
-			ParentWindow = MainFrame.GetParentWindow();
-		}
-
-		TArray<UObject*> Options;
-		Options.SetNum(ImportOptions.Num());
-		for (int32 Index = 0; Index < ImportOptions.Num(); ++Index)
-		{
-			Options[Index] = ImportOptions[Index];
-		}
-
-		TSharedRef<SWindow> Window = SNew(SWindow)
-			.Title(LOCTEXT("DatasmithImportSettingsTitle", "Datasmith Import Options"))
-			.SizingRule(ESizingRule::Autosized);
-
-		float SceneVersion = FDatasmithUtils::GetDatasmithFormatVersionAsFloat();
-		FString FileSDKVersion = FDatasmithUtils::GetEnterpriseVersionAsString();
-
-		TSharedPtr<SDatasmithOptionsWindow> OptionsWindow;
-		Window->SetContent
-		(
-			SAssignNew(OptionsWindow, SDatasmithOptionsWindow)
-			.ImportOptions(Options)
-			.WidgetWindow(Window)
-			// note: Spacing in text below is intentional for text alignment
-			.FileNameText(FText::Format(LOCTEXT("DatasmithImportSettingsFileName", "  Import File  :    {0}"), FText::FromString(FPaths::GetCleanFilename(FilePath))))
-			.FilePathText(FText::FromString(FilePath))
-			.FileFormatVersion(SceneVersion)
-			.FileSDKVersion(FText::FromString(FileSDKVersion))
-//			.PackagePathText(FText::Format(LOCTEXT("DatasmithImportSettingsPackagePath", "  Import To   :    {0}"), FText::FromString(PackagePath)))
-			.ProceedButtonLabel(LOCTEXT("DatasmithOptionWindow_ImportCurLevel", "Import"))
-			.ProceedButtonTooltip(LOCTEXT("DatasmithOptionWindow_ImportCurLevel_ToolTip", "Import the file through Interchange and add to the current Level"))
-			.CancelButtonLabel(LOCTEXT("DatasmithOptionWindow_Cancel", "Cancel"))
-			.CancelButtonTooltip(LOCTEXT("DatasmithOptionWindow_Cancel_ToolTip", "Cancel importing this file"))
-			.MinDetailHeight(320.f)
-			.MinDetailWidth(450.f)
-		);
-
-		FSlateApplication::Get().AddModalWindow(Window, ParentWindow, false);
-
-		return OptionsWindow->ShouldImport();
-	}
-#endif
-}
-
 TArray<FString> UInterchangeDatasmithTranslator::GetSupportedFormats() const
 {
 	const TArray<FString> DatasmithFormats = FDatasmithTranslatorManager::Get().GetSupportedFormats();
@@ -191,9 +130,8 @@ bool UInterchangeDatasmithTranslator::Translate(UInterchangeBaseNodeContainer& B
 					AsyncMode = EAsyncExecution::TaskGraphMainThread;
 				}
 
-				TArray<TObjectPtr<UDatasmithOptionsBase>> OptionArray;
-				OptionArray.Add(CachedSettings->ImportOptions);
-				DatasmithTranslator->SetSceneImportOptions(OptionArray);
+				DatasmithTranslator->SetSceneImportOptions({ CachedSettings->DatasmithOption });
+				CachedSettings->DatasmithOption->SaveConfig();
 			}
 		}
 
@@ -204,6 +142,19 @@ bool UInterchangeDatasmithTranslator::Translate(UInterchangeBaseNodeContainer& B
 		{
 			return false;
 		}
+
+	}
+
+	// Add container for static mesh's additional data
+	{
+		const FString SceneName = DatasmithScene->GetName();
+		const FString StaticMeshDataNodeUid = NodeUtils::ScenePrefix + SceneName + TEXT("_AdditionalData");
+
+		StaticMeshDataNode = NewObject<UDatasmithInterchangeStaticMeshDataNode>(&BaseNodeContainer);
+
+		StaticMeshDataNode->InitializeNode(StaticMeshDataNodeUid, TEXT("StaticMesh_AdditonalData"), EInterchangeNodeContainerType::TranslatedAsset);
+
+		BaseNodeContainer.AddNode(StaticMeshDataNode.Get());
 	}
 
 	// Texture Nodes
@@ -432,7 +383,7 @@ void UInterchangeDatasmithTranslator::HandleDatasmithActor(UInterchangeBaseNodeC
 	BaseNodeContainer.SetNodeParentUid(NodeUid, ParentNodeUid);
 
 	const FTransform ActorTransform = ActorElement->GetRelativeTransform();
-	InterchangeSceneNode->SetCustomLocalTransform(&BaseNodeContainer, ActorElement->GetRelativeTransform());
+	InterchangeSceneNode->SetCustomLocalTransform(&BaseNodeContainer, ActorElement->GetRelativeTransform(), false);
 	// TODO: Layer association + component actors
 
 	if (ActorElement->IsA(EDatasmithElementType::StaticMeshActor))
@@ -524,6 +475,7 @@ UInterchangePhysicalCameraNode* UInterchangeDatasmithTranslator::AddCameraNode(U
 	const float SensorHeight = CameraActor->GetSensorWidth() / CameraActor->GetSensorAspectRatio();
 	CameraNode->SetCustomSensorHeight(SensorHeight);
 
+	// #cad_interchange:
 	// TODO Add properties currently missing from the InterchangeCameraNode:
 	//  - DepthOfField
 	//  - FocusDistance
@@ -770,6 +722,8 @@ TOptional<UE::Interchange::FImportLightProfile> UInterchangeDatasmithTranslator:
 
 TFuture<TOptional<UE::Interchange::FMeshPayloadData>> UInterchangeDatasmithTranslator::GetMeshPayloadData(const FInterchangeMeshPayLoadKey& PayLoadKey, const FTransform& MeshGlobalTransform) const
 {
+	using namespace UE::DatasmithInterchange;
+
 	TPromise<TOptional<UE::Interchange::FMeshPayloadData>> EmptyPromise;
 	EmptyPromise.SetValue(TOptional<UE::Interchange::FMeshPayloadData>());
 
@@ -796,26 +750,10 @@ TFuture<TOptional<UE::Interchange::FMeshPayloadData>> UInterchangeDatasmithTrans
 		{
 			TOptional<UE::Interchange::FMeshPayloadData> Result;
 
-			FDatasmithMeshElementPayload DatasmithMeshPayload;
-			if (LoadedExternalSource->GetAssetTranslator()->LoadStaticMesh(MeshElement.ToSharedRef(), DatasmithMeshPayload))
+			UE::Interchange::FMeshPayloadData StaticMeshPayloadData;
+			if (GetMeshDescription(MeshElement, MeshGlobalTransform, StaticMeshPayloadData))
 			{
-				if (DatasmithMeshPayload.LodMeshes.Num() > 0)
-				{
-					UE::Interchange::FMeshPayloadData StaticMeshPayloadData;
-					StaticMeshPayloadData.MeshDescription = MoveTemp(DatasmithMeshPayload.LodMeshes[0]);
-					if (!FStaticMeshOperations::ValidateAndFixData(StaticMeshPayloadData.MeshDescription, MeshElement->GetName()))
-					{
-						UInterchangeResultError_Generic* ErrorResult = AddMessage<UInterchangeResultError_Generic>();
-						ErrorResult->SourceAssetName = SourceData ? SourceData->GetFilename() : FString();
-						ErrorResult->Text = LOCTEXT("GetMeshPayloadData_ValidateMeshDescriptionFail", "Invalid mesh data (NAN) was found and fix to zero. Mesh render can be bad.");
-					}
-					// Bake the payload mesh, with the provided transform
-					if (!MeshGlobalTransform.Equals(FTransform::Identity))
-					{
-						FStaticMeshOperations::ApplyTransform(StaticMeshPayloadData.MeshDescription, MeshGlobalTransform);
-					}
-					Result.Emplace(MoveTemp(StaticMeshPayloadData));
-				}
+				Result.Emplace(MoveTemp(StaticMeshPayloadData));
 			}
 
 			return Result;
@@ -951,6 +889,16 @@ void UInterchangeDatasmithTranslator::ImportFinish()
 	ElapsedSeconds -= 60.0 * (double)ElapsedMin;
 
 	UE_LOG(LogInterchangeDatasmith, Log, TEXT("Imported %s in [%d min %.3f s]"), *FileName, ElapsedMin, ElapsedSeconds);
+
+	if (LoadedExternalSource.IsValid())
+	{
+		const TSharedPtr<IDatasmithTranslator>& DatasmithTranslator = LoadedExternalSource->GetAssetTranslator();
+		if (DatasmithTranslator)
+		{
+			DatasmithTranslator->UnloadScene();
+		}
+	}
+
 }
 
 
@@ -974,15 +922,10 @@ UInterchangeTranslatorSettings* UInterchangeDatasmithTranslator::GetSettings() c
 			return nullptr;
 		}
 
+		TArray<TObjectPtr<UDatasmithOptionsBase>> DatasmithOptions;
 		const TSharedPtr<IDatasmithTranslator>& DatasmithTranslator = LoadedExternalSource->GetAssetTranslator();
-		if (!DatasmithTranslator)
-		{
-			return nullptr;
-		}
-
-		TArray<TObjectPtr<UDatasmithOptionsBase>> OptionArray;
-		DatasmithTranslator->GetSceneImportOptions(OptionArray);
-		if (OptionArray.Num() == 0)
+		DatasmithTranslator->GetSceneImportOptions(DatasmithOptions);
+		if (DatasmithOptions.Num() == 0)
 		{
 			return nullptr;
 		}
@@ -990,8 +933,13 @@ UInterchangeTranslatorSettings* UInterchangeDatasmithTranslator::GetSettings() c
 		CachedSettings = DuplicateObject<UInterchangeDatasmithTranslatorSettings>(UInterchangeDatasmithTranslatorSettings::StaticClass()->GetDefaultObject<UInterchangeDatasmithTranslatorSettings>(), GetTransientPackage());
 		CachedSettings->SetFlags(RF_Standalone);
 		CachedSettings->ClearInternalFlags(EInternalObjectFlags::Async);
-		CachedSettings->ImportOptions = OptionArray[0];
+
+		// Only the first one is considered
+		CachedSettings->DatasmithOption = DatasmithOptions[0];
+
+		CachedSettings->DatasmithOption->LoadConfig();
 	}
+
 	return CachedSettings;
 }
 
@@ -1008,7 +956,58 @@ void UInterchangeDatasmithTranslator::SetSettings(const UInterchangeTranslatorSe
 		CachedSettings = DuplicateObject<UInterchangeDatasmithTranslatorSettings>(Cast<UInterchangeDatasmithTranslatorSettings>(InterchangeTranslatorSettings), GetTransientPackage());
 		CachedSettings->ClearInternalFlags(EInternalObjectFlags::Async);
 		CachedSettings->SetFlags(RF_Standalone);
+
+		CachedSettings->SaveConfig();
+
+		CachedSettings->DatasmithOption->SaveConfig();
+
+		const TSharedPtr<IDatasmithTranslator>& DatasmithTranslator = LoadedExternalSource->GetAssetTranslator();
+		DatasmithTranslator->SetSceneImportOptions({ CachedSettings->DatasmithOption });
+	}
+}
+
+bool UInterchangeDatasmithTranslator::GetMeshDescription(const TSharedPtr<IDatasmithMeshElement>& MeshElement, const FTransform& MeshGlobalTransform, UE::Interchange::FMeshPayloadData& PayloadData) const
+{
+	using namespace UE::DatasmithInterchange;
+
+	FDatasmithMeshElementPayload DatasmithMeshPayload;
+	if (!LoadedExternalSource->GetAssetTranslator()->LoadStaticMesh(MeshElement.ToSharedRef(), DatasmithMeshPayload))
+	{
+		UInterchangeResultError_Generic* ErrorResult = AddMessage<UInterchangeResultError_Generic>();
+		ErrorResult->SourceAssetName = SourceData ? SourceData->GetFilename() : FString();
+		ErrorResult->Text = FText::Format(LOCTEXT("GetMeshPayloadData_LoadStaticMeshFail", "Failed to load mesh description for mesh element {0}."), FText::FromString(MeshElement->GetName()));
+		return false;
+	}
+
+	if (DatasmithMeshPayload.LodMeshes.Num() > 0)
+	{
+		for (UDatasmithAdditionalData* AdditionalData : DatasmithMeshPayload.AdditionalData)
+		{
+			if (UDatasmithParametricSurfaceData* ParametricSurfaceData = Cast<UDatasmithParametricSurfaceData>(AdditionalData))
+			{
+				const FString MeshNodeUid = NodeUtils::MeshPrefix + MeshElement->GetName();
+				StaticMeshDataNode->AdditionalDataMap.Add(MeshNodeUid, ParametricSurfaceData);
+				break;
+			}
+		}
+
+		if (!FStaticMeshOperations::ValidateAndFixData(DatasmithMeshPayload.LodMeshes[0], MeshElement->GetName()))
+		{
+			UInterchangeResultError_Generic* ErrorResult = AddMessage<UInterchangeResultError_Generic>();
+			ErrorResult->SourceAssetName = SourceData ? SourceData->GetFilename() : FString();
+			ErrorResult->Text = FText::Format(LOCTEXT("GetMeshPayloadData_ValidateMeshDescriptionFail", "Invalid mesh data (NAN) was found and fix to zero. Mesh render can be bad for mesh element {0}."), FText::FromString(MeshElement->GetName()));
+			return false;
+		}
+		// Bake the payload mesh, with the provided transform
+		if (!MeshGlobalTransform.Equals(FTransform::Identity))
+		{
+			FStaticMeshOperations::ApplyTransform(DatasmithMeshPayload.LodMeshes[0], MeshGlobalTransform);
+		}
+
+		PayloadData.MeshDescription = MoveTemp(DatasmithMeshPayload.LodMeshes[0]);
 	}
+
+	return true;
 }
 
 #undef LOCTEXT_NAMESPACE
\ No newline at end of file
diff --git a/Engine/Plugins/Experimental/Enterprise/DatasmithInterchange/Source/DatasmithInterchange/Public/InterchangeDatasmithTranslator.h b/Engine/Plugins/Experimental/Enterprise/DatasmithInterchange/Source/DatasmithInterchange/Public/InterchangeDatasmithTranslator.h
--- a/Engine/Plugins/Experimental/Enterprise/DatasmithInterchange/Source/DatasmithInterchange/Public/InterchangeDatasmithTranslator.h
+++ b/Engine/Plugins/Experimental/Enterprise/DatasmithInterchange/Source/DatasmithInterchange/Public/InterchangeDatasmithTranslator.h
@@ -14,6 +14,8 @@
 #include "Texture/InterchangeTextureLightProfilePayloadInterface.h"
 #include "Scene/InterchangeVariantSetPayloadInterface.h"
 
+#include "DatasmithImportOptions.h"
+
 #include "Async/Async.h"
 #include "ExternalSource.h"
 #include "UObject/GCObjectScopeGuard.h"
@@ -27,9 +29,11 @@ class IDatasmithBaseAnimationElement;
 class IDatasmithCameraActorElement;
 class IDatasmithLightActorElement;
 class IDatasmithDecalActorElement;
+class IDatasmithMeshElement;
 class IDatasmithScene;
 class IDatasmithTransformAnimationElement;
 class UDatasmithOptionsBase;
+class UDatasmithInterchangeStaticMeshDataNode;
 class UInterchangePhysicalCameraNode;
 class UInterchangeBaseLightNode;
 class UInterchangeDecalNode;
@@ -55,14 +59,22 @@ namespace UE::DatasmithInterchange::AnimUtils
 	extern bool GetAnimationPayloadData(const IDatasmithBaseAnimationElement& AnimationElement, float FrameRate, EInterchangeAnimationPayLoadType PayLoadType, UE::Interchange::FAnimationPayloadData& PayLoadData);
 }
 
+UENUM()
+enum class EInterchangeMesherType : uint8
+{
+	UseCADKernel UMETA(DisplayName = "Use CADKernel"),
+	UseTechSoft,
+	UseNativeTessellator,
+};
+
 UCLASS(BlueprintType, editinlinenew, MinimalAPI)
 class UInterchangeDatasmithTranslatorSettings : public UInterchangeTranslatorSettings
 {
 	GENERATED_BODY()
 
 public:
-	UPROPERTY(EditAnywhere, Instanced, BlueprintReadWrite, Category = "Datasmith Interchange", meta = (ShowOnlyInnerProperties))
-	TObjectPtr<UDatasmithOptionsBase> ImportOptions;
+	UPROPERTY(EditAnywhere, BlueprintReadWrite, Category = "Datasmith Options")
+	TObjectPtr<UDatasmithOptionsBase> DatasmithOption;
 };
 
 UCLASS(BlueprintType, Experimental)
@@ -136,6 +148,8 @@ private:
 
 	void ProcessIesProfile(UInterchangeBaseNodeContainer& BaseNodeContainer, const IDatasmithLightActorElement& LightElement, UInterchangeLightNode* LightNode) const;
 
+	bool GetMeshDescription(const TSharedPtr<IDatasmithMeshElement>& MeshElement, const FTransform& MeshGlobalTransform, UE::Interchange::FMeshPayloadData& PayloadData) const;
+
 	mutable TSharedPtr<UE::DatasmithImporter::FExternalSource> LoadedExternalSource;
 
 	mutable uint64 StartTime = 0;
@@ -143,5 +157,7 @@ private:
 	mutable TObjectPtr<UInterchangeDatasmithTranslatorSettings> CachedSettings = nullptr;
 
 	mutable TMap<FString, UE::DatasmithInterchange::AnimUtils::FAnimationPayloadDesc> AnimationPayLoadMapping;
+
+	mutable TObjectPtr<UDatasmithInterchangeStaticMeshDataNode> StaticMeshDataNode;
 	mutable EAsyncExecution AsyncMode = EAsyncExecution::TaskGraph;
 };
diff --git a/Engine/Plugins/Interchange/Runtime/Source/Pipelines/Public/InterchangeGenericMeshPipeline.h b/Engine/Plugins/Interchange/Runtime/Source/Pipelines/Public/InterchangeGenericMeshPipeline.h
--- a/Engine/Plugins/Interchange/Runtime/Source/Pipelines/Public/InterchangeGenericMeshPipeline.h
+++ b/Engine/Plugins/Interchange/Runtime/Source/Pipelines/Public/InterchangeGenericMeshPipeline.h
@@ -365,8 +365,7 @@ private:
 		, bool& bContainStaticMeshAnimationNode
 		, bool& bIgnoreStaticMeshes) const;
 
-private:
-
+protected:
 	UInterchangeBaseNodeContainer* BaseNodeContainer = nullptr;
 	TArray<const UInterchangeSourceData*> SourceDatas;
 
