Fixed socket issue with CADWOrker: It was a timing issue. It takes more than 3 seconds for the worker to ping the editor. Moved the timeout from 3 to 6 sseconds
Fixed crash when tessellating a specific surface: A polyline of 1 point was generated. This was not tested against.
Fixed crash on reimport: The Datasmith external source was required to set the Datasmith translator.

#rnx

diff --git a/Engine/Plugins/Enterprise/DatasmithCADImporter/Source/DatasmithCADTranslator/Private/DatasmithCADTranslator.cpp b/Engine/Plugins/Enterprise/DatasmithCADImporter/Source/DatasmithCADTranslator/Private/DatasmithCADTranslator.cpp
--- a/Engine/Plugins/Enterprise/DatasmithCADImporter/Source/DatasmithCADTranslator/Private/DatasmithCADTranslator.cpp
+++ b/Engine/Plugins/Enterprise/DatasmithCADImporter/Source/DatasmithCADTranslator/Private/DatasmithCADTranslator.cpp
@@ -54,7 +54,7 @@ void FDatasmithCADTranslator::Initialize(FDatasmithTranslatorCapabilities& OutCa
 	}
 	OutCapabilities.bIsEnabled = true;
 
-#ifndef CAD_TRANSLATOR_DEBUG
+#ifdef CAD_TRANSLATOR_DEBUG
 	OutCapabilities.bParallelLoadStaticMeshSupported = true;
 #endif
 	OutCapabilities.SupportedFileFormats.Add(FFileFormatInfo{ TEXT("CATPart"), TEXT("CATIA Part files") });
diff --git a/Engine/Plugins/Enterprise/DatasmithCADImporter/Source/DatasmithDispatcher/Private/DatasmithDispatcher.cpp b/Engine/Plugins/Enterprise/DatasmithCADImporter/Source/DatasmithDispatcher/Private/DatasmithDispatcher.cpp
--- a/Engine/Plugins/Enterprise/DatasmithCADImporter/Source/DatasmithDispatcher/Private/DatasmithDispatcher.cpp
+++ b/Engine/Plugins/Enterprise/DatasmithCADImporter/Source/DatasmithDispatcher/Private/DatasmithDispatcher.cpp
@@ -394,11 +394,6 @@ void FDatasmithDispatcher::SpawnHandlers()
 	}
 }
 
-int32 FDatasmithDispatcher::GetAliveHandlerCount()
-{
-	return Algo::CountIf(WorkerHandlers, [](const FDatasmithWorkerHandler& Handler) { return Handler.IsAlive(); });
-}
-
 void FDatasmithDispatcher::CloseHandlers()
 {
 	for (FDatasmithWorkerHandler& Handler : WorkerHandlers)
diff --git a/Engine/Plugins/Enterprise/DatasmithCADImporter/Source/DatasmithDispatcher/Private/DatasmithDispatcherNetworking.cpp b/Engine/Plugins/Enterprise/DatasmithCADImporter/Source/DatasmithDispatcher/Private/DatasmithDispatcherNetworking.cpp
--- a/Engine/Plugins/Enterprise/DatasmithCADImporter/Source/DatasmithDispatcher/Private/DatasmithDispatcherNetworking.cpp
+++ b/Engine/Plugins/Enterprise/DatasmithCADImporter/Source/DatasmithDispatcher/Private/DatasmithDispatcherNetworking.cpp
@@ -241,8 +241,6 @@ FNetworkServerNode::FNetworkServerNode()
 	TSharedRef<FInternetAddr> InternetAddress = SocketSubsystem->CreateInternetAddr();
 	InternetAddress->SetLoopbackAddress();
 
-	FString a = InternetAddress->ToString(true);
-
 	int32 BoundPort = SocketSubsystem->BindNextPort(ListeningSocket, *InternetAddress, 1e3, 1);
 	if (BoundPort == 0)
 	{
diff --git a/Engine/Plugins/Enterprise/DatasmithCADImporter/Source/DatasmithDispatcher/Public/DatasmithDispatcher.h b/Engine/Plugins/Enterprise/DatasmithCADImporter/Source/DatasmithDispatcher/Public/DatasmithDispatcher.h
--- a/Engine/Plugins/Enterprise/DatasmithCADImporter/Source/DatasmithDispatcher/Public/DatasmithDispatcher.h
+++ b/Engine/Plugins/Enterprise/DatasmithCADImporter/Source/DatasmithDispatcher/Public/DatasmithDispatcher.h
@@ -38,7 +38,6 @@ public:
 private:
 	void SpawnHandlers();
 	int32 GetNextWorkerId();
-	int32 GetAliveHandlerCount();
 	void CloseHandlers();
 
 	void ProcessLocal();
diff --git a/Engine/Plugins/Enterprise/DatasmithCADImporter/Source/DatasmithDispatcher/Public/DatasmithDispatcherConfig.h b/Engine/Plugins/Enterprise/DatasmithCADImporter/Source/DatasmithDispatcher/Public/DatasmithDispatcherConfig.h
--- a/Engine/Plugins/Enterprise/DatasmithCADImporter/Source/DatasmithDispatcher/Public/DatasmithDispatcherConfig.h
+++ b/Engine/Plugins/Enterprise/DatasmithCADImporter/Source/DatasmithDispatcher/Public/DatasmithDispatcherConfig.h
@@ -10,7 +10,7 @@ namespace Config
 {
 
 // Server side timeout for connection with the client
-constexpr double AcceptTimeout_s = 3.0;
+constexpr double AcceptTimeout_s = 6.0;
 
 // Client side timeout for connection with the server
 constexpr double ConnectTimeout_s = 3.0;
diff --git a/Engine/Plugins/Experimental/Enterprise/DatasmithInterchange/Source/DatasmithInterchange/Private/InterchangeDatasmithTranslator.cpp b/Engine/Plugins/Experimental/Enterprise/DatasmithInterchange/Source/DatasmithInterchange/Private/InterchangeDatasmithTranslator.cpp
--- a/Engine/Plugins/Experimental/Enterprise/DatasmithInterchange/Source/DatasmithInterchange/Private/InterchangeDatasmithTranslator.cpp
+++ b/Engine/Plugins/Experimental/Enterprise/DatasmithInterchange/Source/DatasmithInterchange/Private/InterchangeDatasmithTranslator.cpp
@@ -890,6 +890,10 @@ void UInterchangeDatasmithTranslator::ImportFinish()
 
 	UE_LOG(LogInterchangeDatasmith, Log, TEXT("Imported %s in [%d min %.3f s]"), *FileName, ElapsedMin, ElapsedSeconds);
 
+	// Remove dependency on created static meshes
+	StaticMeshDataNode->AdditionalDataMap.Reset();
+	StaticMeshDataNode = nullptr;
+
 	if (LoadedExternalSource.IsValid())
 	{
 		const TSharedPtr<IDatasmithTranslator>& DatasmithTranslator = LoadedExternalSource->GetAssetTranslator();
@@ -945,6 +949,8 @@ UInterchangeTranslatorSettings* UInterchangeDatasmithTranslator::GetSettings() c
 
 void UInterchangeDatasmithTranslator::SetSettings(const UInterchangeTranslatorSettings* InterchangeTranslatorSettings)
 {
+	using namespace UE::DatasmithImporter;
+
 	if (CachedSettings)
 	{
 		CachedSettings->ClearFlags(RF_Standalone);
@@ -961,8 +967,19 @@ void UInterchangeDatasmithTranslator::SetSettings(const UInterchangeTranslatorSe
 
 		CachedSettings->DatasmithOption->SaveConfig();
 
-		const TSharedPtr<IDatasmithTranslator>& DatasmithTranslator = LoadedExternalSource->GetAssetTranslator();
-		DatasmithTranslator->SetSceneImportOptions({ CachedSettings->DatasmithOption });
+		if (!LoadedExternalSource.IsValid())
+		{
+			FString FilePath = FPaths::ConvertRelativePathToFull(SourceData->GetFilename());
+			FileName = FPaths::GetCleanFilename(FilePath);
+			const FSourceUri FileNameUri = FSourceUri::FromFilePath(FilePath);
+			LoadedExternalSource = IExternalSourceModule::GetOrCreateExternalSource(FileNameUri);
+		}
+
+		if (LoadedExternalSource.IsValid() && LoadedExternalSource->IsAvailable())
+		{
+			const TSharedPtr<IDatasmithTranslator>& DatasmithTranslator = LoadedExternalSource->GetAssetTranslator();
+			DatasmithTranslator->SetSceneImportOptions({ CachedSettings->DatasmithOption });
+		}
 	}
 }
 
diff --git a/Engine/Source/Programs/Enterprise/Datasmith/DatasmithCADWorker/Private/DatasmithCADWorkerImpl.cpp b/Engine/Source/Programs/Enterprise/Datasmith/DatasmithCADWorker/Private/DatasmithCADWorkerImpl.cpp
--- a/Engine/Source/Programs/Enterprise/Datasmith/DatasmithCADWorker/Private/DatasmithCADWorkerImpl.cpp
+++ b/Engine/Source/Programs/Enterprise/Datasmith/DatasmithCADWorker/Private/DatasmithCADWorkerImpl.cpp
@@ -12,6 +12,7 @@
 #include "Misc/CoreDelegates.h"
 #include "Misc/FileHelper.h"
 #include "Misc/Paths.h"
+#include "Misc/OutputDeviceRedirector.h"
 #include "SocketSubsystem.h"
 #include "Sockets.h"
 #include "Tasks/Task.h"
@@ -27,10 +28,6 @@ FDatasmithCADWorkerImpl::FDatasmithCADWorkerImpl(int32 InServerPID, int32 InServ
 	, EnginePluginsPath(InEnginePluginsPath)
 	, CachePath(InCachePath)
 	, PingStartCycle(0)
-{
-}
-
-bool FDatasmithCADWorkerImpl::Run()
 {
 	UE_LOG(LogDatasmithCADWorker, Verbose, TEXT("connect to %d..."), ServerPort);
 	bool bConnected = NetworkInterface.Connect(TEXT("Datasmith CAD Worker"), ServerPort, Config::ConnectTimeout_s);
@@ -42,10 +39,18 @@ bool FDatasmithCADWorkerImpl::Run()
 	else
 	{
 		UE_LOG(LogDatasmithCADWorker, Error, TEXT("Server connection failure. exit"));
-		return false;
+		return;
 	}
 
 	InitiatePing();
+}
+
+bool FDatasmithCADWorkerImpl::Run()
+{
+	if (!CommandIO.IsValid())
+	{
+		return false;
+	}
 
 	bool bIsRunning = true;
 	while (bIsRunning)
@@ -98,6 +103,7 @@ bool FDatasmithCADWorkerImpl::Run()
 void FDatasmithCADWorkerImpl::InitiatePing()
 {
 	PingStartCycle = FPlatformTime::Cycles64();
+	UE_LOG(LogDatasmithCADWorker, Verbose, TEXT("InitiatePing..."));
 	FPingCommand Ping;
 	CommandIO.SendCommand(Ping, Config::SendCommandTimeout_s);
 }
diff --git a/Engine/Source/Runtime/Datasmith/CADKernel/Base/Public/Geo/Curves/RestrictionCurve.h b/Engine/Source/Runtime/Datasmith/CADKernel/Base/Public/Geo/Curves/RestrictionCurve.h
--- a/Engine/Source/Runtime/Datasmith/CADKernel/Base/Public/Geo/Curves/RestrictionCurve.h
+++ b/Engine/Source/Runtime/Datasmith/CADKernel/Base/Public/Geo/Curves/RestrictionCurve.h
@@ -169,6 +169,13 @@ public:
 	 */
 	void CheckIfDegenerated(const FLinearBoundary& InBoundary, bool& bDegeneration2D, bool& bDegeneration3D, double& Length3D) const
 	{
+		if (!ensure(Polyline.Coordinates.Num() > 1))
+		{
+			bDegeneration2D = true;
+			bDegeneration3D = true;
+			Length3D = 0.;
+			return;
+		}
 		double Tolerance = GetCarrierSurface()->Get3DTolerance();
 		const FSurfacicTolerance& Tolerances2D = GetCarrierSurface()->GetIsoTolerances();
 		Polyline.CheckIfDegenerated(Tolerance, Tolerances2D, InBoundary, bDegeneration2D, bDegeneration3D, Length3D);
