Fixed deadlock when UInterchangeDatasmithTranslator::GetMeshPayloadData  is called on the GameThread while using a non-threadsafe Datasmith translator.
Solution: THe GetMeshPayloadData method is executed on the spot.

#rnx

diff --git a/Engine/Plugins/Experimental/Enterprise/DatasmithInterchange/Source/DatasmithInterchange/Private/InterchangeDatasmithTranslator.cpp b/Engine/Plugins/Experimental/Enterprise/DatasmithInterchange/Source/DatasmithInterchange/Private/InterchangeDatasmithTranslator.cpp
--- a/Engine/Plugins/Experimental/Enterprise/DatasmithInterchange/Source/DatasmithInterchange/Private/InterchangeDatasmithTranslator.cpp
+++ b/Engine/Plugins/Experimental/Enterprise/DatasmithInterchange/Source/DatasmithInterchange/Private/InterchangeDatasmithTranslator.cpp
@@ -746,6 +746,26 @@ TFuture<TOptional<UE::Interchange::FMeshPayloadData>> UInterchangeDatasmithTrans
 		return EmptyPromise.GetFuture();
 	}
 
+	// If on the GameThread and Datasmith translator cannot be parallelized
+	// just run it on the spot.
+	if (IsInGameThread() && AsyncMode == EAsyncExecution::TaskGraphMainThread)
+	{
+		TOptional<UE::Interchange::FMeshPayloadData> Result;
+
+		UE::Interchange::FMeshPayloadData StaticMeshPayloadData;
+		if (GetMeshDescription(MeshElement, MeshGlobalTransform, StaticMeshPayloadData))
+		{
+			Result.Emplace(MoveTemp(StaticMeshPayloadData));
+
+			TPromise<TOptional<UE::Interchange::FMeshPayloadData>> Promise;
+			Promise.SetValue(Result);
+
+			return Promise.GetFuture();
+		}
+
+		return EmptyPromise.GetFuture();
+	}
+
 	return Async(AsyncMode, [this, MeshElement = MoveTemp(MeshElement), MeshGlobalTransform]
 		{
 			TOptional<UE::Interchange::FMeshPayloadData> Result;
