Fixed crash when opening a project after enabling DatasmithInterchange plugin on Mac

#rnx

diff --git a/Engine/Plugins/Experimental/Enterprise/DatasmithInterchange/Source/DatasmithInterchange/Private/DatasmithInterchangeModule.cpp b/Engine/Plugins/Experimental/Enterprise/DatasmithInterchange/Source/DatasmithInterchange/Private/DatasmithInterchangeModule.cpp
--- a/Engine/Plugins/Experimental/Enterprise/DatasmithInterchange/Source/DatasmithInterchange/Private/DatasmithInterchangeModule.cpp
+++ b/Engine/Plugins/Experimental/Enterprise/DatasmithInterchange/Source/DatasmithInterchange/Private/DatasmithInterchangeModule.cpp
@@ -39,6 +39,30 @@ class FDatasmithInterchangeModule : public IDatasmithInterchangeModule
 {
 public:
 	virtual void StartupModule() override
+	{
+		FCoreDelegates::OnPostEngineInit.AddRaw(this, &FDatasmithInterchangeModule::OnPostEngineInit);
+	}
+
+	virtual void ShutdownModule() override
+	{
+		FCoreDelegates::OnPostEngineInit.RemoveAll(this);
+
+		UE::DatasmithInterchange::FDatasmithReferenceMaterialManager::Destroy();
+		
+#if WITH_EDITOR
+		FPropertyEditorModule* PropertyEditorModule = FModuleManager::GetModulePtr<FPropertyEditorModule>("PropertyEditor");
+		if (PropertyEditorModule)
+		{
+			for (FName ClassName : ClassesToUnregisterOnShutdown)
+			{
+				PropertyEditorModule->UnregisterCustomClassLayout(ClassName);
+			}
+		}
+		ClassesToUnregisterOnShutdown.Empty();
+#endif
+	}
+
+	void OnPostEngineInit()
 	{
 		using namespace UE::DatasmithInterchange;
 
@@ -74,7 +98,7 @@ public:
 		FDatasmithReferenceMaterialManager::Get().RegisterSelector(TEXT("SketchUp"), MakeShared< FDatasmithSketchUpMaterialSelector >());
 		FDatasmithReferenceMaterialManager::Get().RegisterSelector(TEXT("CityEngine"), MakeShared< FDatasmithCityEngineMaterialSelector >());
 		FDatasmithReferenceMaterialManager::Get().RegisterSelector(TEXT("StdMaterial"), MakeShared< FDatasmithStdMaterialSelector >());
-		
+
 #if WITH_EDITOR
 		ClassesToUnregisterOnShutdown.Reset();
 		// Register details customizations
@@ -85,23 +109,6 @@ public:
 #endif
 	}
 
-	virtual void ShutdownModule() override
-	{
-		UE::DatasmithInterchange::FDatasmithReferenceMaterialManager::Destroy();
-		
-#if WITH_EDITOR
-		FPropertyEditorModule* PropertyEditorModule = FModuleManager::GetModulePtr<FPropertyEditorModule>("PropertyEditor");
-		if (PropertyEditorModule)
-		{
-			for (FName ClassName : ClassesToUnregisterOnShutdown)
-			{
-				PropertyEditorModule->UnregisterCustomClassLayout(ClassName);
-			}
-		}
-		ClassesToUnregisterOnShutdown.Empty();
-#endif
-	}
-
 #if WITH_EDITOR
 	TArray<FName> ClassesToUnregisterOnShutdown;
 #endif
