Fixed issue with reimport policy not applied
Root cause: Missing update of import context

#rnx

diff --git a/Engine/Plugins/Experimental/Enterprise/DatasmithInterchange/Source/DatasmithInterchange/Private/InterchangeDatasmithPipeline.cpp b/Engine/Plugins/Experimental/Enterprise/DatasmithInterchange/Source/DatasmithInterchange/Private/InterchangeDatasmithPipeline.cpp
--- a/Engine/Plugins/Experimental/Enterprise/DatasmithInterchange/Source/DatasmithInterchange/Private/InterchangeDatasmithPipeline.cpp
+++ b/Engine/Plugins/Experimental/Enterprise/DatasmithInterchange/Source/DatasmithInterchange/Private/InterchangeDatasmithPipeline.cpp
@@ -135,15 +135,8 @@ void UInterchangeDatasmithPipeline::AdjustSettingsForContext(const FInterchangeP
 	{
 		AnimationPipeline->AdjustSettingsForContext(ContextParams);
 	}
-}
 
-void UInterchangeDatasmithPipeline::PostDuplicate(bool bDuplicateForPIE)
-{
-	// Only adjust settings if there is anything cached.
-	if (CacheContextParam.ContextType != EInterchangePipelineContext::None)
-	{
-		AdjustSettingsForContext(CacheContextParam);
-	}
+	CacheContextParam = ContextParams;
 }
 
 void UInterchangeDatasmithPipeline::ExecutePipeline(UInterchangeBaseNodeContainer* InBaseNodeContainer, const TArray<UInterchangeSourceData*>& SourceDatas, const FString& ContentBasePath)
@@ -159,6 +152,7 @@ void UInterchangeDatasmithPipeline::ExecutePipeline(UInterchangeBaseNodeContaine
 	{
 		if (Pipeline)
 		{
+			Pipeline->AdjustSettingsForContext(CacheContextParam);
 			Pipeline->SetResultsContainer(this->Results);
 			Pipeline->ScriptedExecutePipeline(this->BaseNodeContainer, SourceDatas, ContentBasePath);
 		}
diff --git a/Engine/Plugins/Experimental/Enterprise/DatasmithInterchange/Source/DatasmithInterchange/Public/InterchangeDatasmithPipeline.h b/Engine/Plugins/Experimental/Enterprise/DatasmithInterchange/Source/DatasmithInterchange/Public/InterchangeDatasmithPipeline.h
--- a/Engine/Plugins/Experimental/Enterprise/DatasmithInterchange/Source/DatasmithInterchange/Public/InterchangeDatasmithPipeline.h
+++ b/Engine/Plugins/Experimental/Enterprise/DatasmithInterchange/Source/DatasmithInterchange/Public/InterchangeDatasmithPipeline.h
@@ -56,11 +56,6 @@ class DATASMITHINTERCHANGE_API UInterchangeDatasmithPipeline : public UInterchan
 	UPROPERTY(EditAnywhere, Category = "Materials")
 	bool bCreateMaterialReferencesFolders = true;
 
-public:
-	/** Begin UObject overrides*/
-	virtual void PostDuplicate(bool bDuplicateForPIE) override;
-	/** End UObject overrides */
-
 protected:
 	virtual void ExecutePipeline(UInterchangeBaseNodeContainer* InBaseNodeContainer, const TArray<UInterchangeSourceData*>& SourceDatas, const FString& ContentBasePath) override;
 	virtual void ExecutePostImportPipeline(const UInterchangeBaseNodeContainer* InBaseNodeContainer, const FString& NodeKey, UObject* CreatedAsset, bool bIsAReimport) override;
