[Interchange] Fixed race condition in UInterchangeDatasmithTranslator

- fixes Editor crash while parsing data after Import (using CATProduct)

diff --git a/Engine/Plugins/Experimental/Enterprise/DatasmithInterchange/Source/DatasmithInterchange/Private/InterchangeDatasmithTranslator.cpp b/Engine/Plugins/Experimental/Enterprise/DatasmithInterchange/Source/DatasmithInterchange/Private/InterchangeDatasmithTranslator.cpp
--- a/Engine/Plugins/Experimental/Enterprise/DatasmithInterchange/Source/DatasmithInterchange/Private/InterchangeDatasmithTranslator.cpp
+++ b/Engine/Plugins/Experimental/Enterprise/DatasmithInterchange/Source/DatasmithInterchange/Private/InterchangeDatasmithTranslator.cpp
@@ -1006,6 +1006,8 @@ bool UInterchangeDatasmithTranslator::GetMeshDescription(const TSharedPtr<IDatas
 		{
 			if (UDatasmithParametricSurfaceData* ParametricSurfaceData = Cast<UDatasmithParametricSurfaceData>(AdditionalData))
 			{
+				FWriteScopeLock ReconnectionScopeLock(StaticMeshDataNodeLock);
+
 				const FString MeshNodeUid = NodeUtils::MeshPrefix + MeshElement->GetName();
 				StaticMeshDataNode->AdditionalDataMap.Add(MeshNodeUid, ParametricSurfaceData);
 				break;
diff --git a/Engine/Plugins/Experimental/Enterprise/DatasmithInterchange/Source/DatasmithInterchange/Public/InterchangeDatasmithTranslator.h b/Engine/Plugins/Experimental/Enterprise/DatasmithInterchange/Source/DatasmithInterchange/Public/InterchangeDatasmithTranslator.h
--- a/Engine/Plugins/Experimental/Enterprise/DatasmithInterchange/Source/DatasmithInterchange/Public/InterchangeDatasmithTranslator.h
+++ b/Engine/Plugins/Experimental/Enterprise/DatasmithInterchange/Source/DatasmithInterchange/Public/InterchangeDatasmithTranslator.h
@@ -160,6 +160,7 @@ private:
 
 	mutable TMap<FString, UE::DatasmithInterchange::AnimUtils::FAnimationPayloadDesc> AnimationPayLoadMapping;
 
+	mutable FRWLock StaticMeshDataNodeLock;
 	mutable TObjectPtr<UDatasmithInterchangeStaticMeshDataNode> StaticMeshDataNode;
 	mutable EAsyncExecution AsyncMode = EAsyncExecution::TaskGraph;
 };
