uDatasmith Interchange Analytics Improvements
- Application Name, Version and Vendor are now passed in the Analytics

#rnx

diff --git a/Engine/Plugins/Experimental/Enterprise/DatasmithInterchange/Source/DatasmithInterchange/Private/InterchangeDatasmithTranslator.cpp b/Engine/Plugins/Experimental/Enterprise/DatasmithInterchange/Source/DatasmithInterchange/Private/InterchangeDatasmithTranslator.cpp
--- a/Engine/Plugins/Experimental/Enterprise/DatasmithInterchange/Source/DatasmithInterchange/Private/InterchangeDatasmithTranslator.cpp
+++ b/Engine/Plugins/Experimental/Enterprise/DatasmithInterchange/Source/DatasmithInterchange/Private/InterchangeDatasmithTranslator.cpp
@@ -38,6 +38,7 @@
 #include "InterchangeTextureLightProfileFactoryNode.h"
 #include "InterchangeVariantSetNode.h"
 #include "StaticMeshOperations.h"
+#include "Nodes/InterchangeSourceNode.h"
 
 #include "Async/ParallelFor.h"
 #include "HAL/ConsoleManager.h"
@@ -230,6 +231,17 @@ bool UInterchangeDatasmithTranslator::Translate(UInterchangeBaseNodeContainer& B
 
 	}
 
+	{
+		// File Creator Meta Data Information
+		if (UInterchangeSourceNode* SourceNode = UInterchangeSourceNode::FindOrCreateUniqueInstance(&BaseNodeContainer))
+		{
+			using namespace UE::Interchange;
+			SourceNode->SetExtraInformation(FSourceNodeExtraInfoStaticData::GetApplicationVendorExtraInfoKey(), DatasmithScene->GetVendor());
+			SourceNode->SetExtraInformation(FSourceNodeExtraInfoStaticData::GetApplicationNameExtraInfoKey(), DatasmithScene->GetProductName());
+			SourceNode->SetExtraInformation(FSourceNodeExtraInfoStaticData::GetApplicationVersionExtraInfoKey(), DatasmithScene->GetProductVersion());
+		}
+	}
+
 	// Add container for static mesh's additional data
 	{
 		const FString SceneName = DatasmithScene->GetName();
