Retrofit 38509813 to 5.5.4 to fix UE-220975

#rnx

diff --git a/Engine/Plugins/Experimental/Enterprise/DatasmithInterchange/Source/DatasmithInterchange/Private/InterchangeDatasmithTranslator.cpp b/Engine/Plugins/Experimental/Enterprise/DatasmithInterchange/Source/DatasmithInterchange/Private/InterchangeDatasmithTranslator.cpp
--- a/Engine/Plugins/Experimental/Enterprise/DatasmithInterchange/Source/DatasmithInterchange/Private/InterchangeDatasmithTranslator.cpp
+++ b/Engine/Plugins/Experimental/Enterprise/DatasmithInterchange/Source/DatasmithInterchange/Private/InterchangeDatasmithTranslator.cpp
@@ -125,6 +125,8 @@ static TSet<FString> ExcludedFormats{
 	TEXT("3dm"),
 };
 
+FRWLock UInterchangeDatasmithTranslator::StaticMeshDataNodeLock;
+
 TArray<FString> UInterchangeDatasmithTranslator::GetSupportedFormats() const
 {
 	IConsoleVariable* CVar = IConsoleManager::Get().FindConsoleVariable(TEXT("Interchange.FeatureFlags.Import.CAD"), false);
diff --git a/Engine/Plugins/Experimental/Enterprise/DatasmithInterchange/Source/DatasmithInterchange/Private/InterchangeDatasmithTranslator.h b/Engine/Plugins/Experimental/Enterprise/DatasmithInterchange/Source/DatasmithInterchange/Private/InterchangeDatasmithTranslator.h
--- a/Engine/Plugins/Experimental/Enterprise/DatasmithInterchange/Source/DatasmithInterchange/Private/InterchangeDatasmithTranslator.h
+++ b/Engine/Plugins/Experimental/Enterprise/DatasmithInterchange/Source/DatasmithInterchange/Private/InterchangeDatasmithTranslator.h
@@ -160,7 +160,7 @@ private:
 
 	mutable TMap<FString, UE::DatasmithInterchange::AnimUtils::FAnimationPayloadDesc> AnimationPayLoadMapping;
 
-	mutable FRWLock StaticMeshDataNodeLock;
+	static FRWLock StaticMeshDataNodeLock;
 	mutable TObjectPtr<UDatasmithInterchangeStaticMeshDataNode> StaticMeshDataNode;
 	mutable EAsyncExecution AsyncMode = EAsyncExecution::TaskGraph;
 };
