Interchange MaterialX/USD:
- Fix data race condition

diff --git a/Engine/Plugins/Experimental/Enterprise/DatasmithInterchange/Source/DatasmithInterchange/Private/InterchangeDatasmithTranslator.cpp b/Engine/Plugins/Experimental/Enterprise/DatasmithInterchange/Source/DatasmithInterchange/Private/InterchangeDatasmithTranslator.cpp
--- a/Engine/Plugins/Experimental/Enterprise/DatasmithInterchange/Source/DatasmithInterchange/Private/InterchangeDatasmithTranslator.cpp
+++ b/Engine/Plugins/Experimental/Enterprise/DatasmithInterchange/Source/DatasmithInterchange/Private/InterchangeDatasmithTranslator.cpp
@@ -36,6 +36,7 @@
 #include "InterchangeTexture2DNode.h"
 #include "InterchangeTextureLightProfileNode.h"
 #include "InterchangeTextureLightProfileFactoryNode.h"
+#include "InterchangeTranslatorHelper.h"
 #include "InterchangeVariantSetNode.h"
 #include "StaticMeshOperations.h"
 #include "Nodes/InterchangeSourceNode.h"
@@ -759,27 +760,13 @@ TOptional<UE::Interchange::FImportImage> UInterchangeDatasmithTranslator::GetTex
 		return TOptional<UE::Interchange::FImportImage>();
 	}
 
-	UInterchangeSourceData* PayloadSourceData = UInterchangeManager::GetInterchangeManager().CreateSourceData(TextureElement->GetFile());
-	FGCObjectScopeGuard ScopedSourceData(PayloadSourceData);
-
-	if (!PayloadSourceData)
-	{
-		return TOptional<UE::Interchange::FImportImage>();
-	}
-
-	PayloadSourceData->ClearInternalFlags(EInternalObjectFlags::Async);
-
-	UInterchangeTranslatorBase* SourceTranslator = UInterchangeManager::GetInterchangeManager().GetTranslatorForSourceData(PayloadSourceData);
-	FGCObjectScopeGuard ScopedSourceTranslator(SourceTranslator);
-	const IInterchangeTexturePayloadInterface* TextureTranslator = Cast< IInterchangeTexturePayloadInterface >(SourceTranslator);
+	UE::Interchange::Private::FScopedTranslator ScopedTranslator(TextureElement->GetFile(), Results);
+	const IInterchangeTexturePayloadInterface* TextureTranslator = ScopedTranslator.GetPayLoadInterface< IInterchangeTexturePayloadInterface >();
 	if (!ensure(TextureTranslator))
 	{
 		return TOptional<UE::Interchange::FImportImage>();
 	}
 
-	SourceTranslator->ClearInternalFlags(EInternalObjectFlags::Async);
-	SourceTranslator->SetResultsContainer(Results);
-
 	AlternateTexturePath = TextureElement->GetFile();
 
 	return TextureTranslator->GetTexturePayloadData(PayloadKey, AlternateTexturePath);
@@ -793,25 +780,13 @@ TOptional<UE::Interchange::FImportLightProfile> UInterchangeDatasmithTranslator:
 	}
 
 	UInterchangeSourceData* PayloadSourceData = UInterchangeManager::GetInterchangeManager().CreateSourceData(PayloadKey);
-	FGCObjectScopeGuard ScopedSourceData(PayloadSourceData);
-	if (!PayloadSourceData)
-	{
-		return TOptional<UE::Interchange::FImportLightProfile>();
-	}
-
-	PayloadSourceData->ClearInternalFlags(EInternalObjectFlags::Async);
-
-	UInterchangeTranslatorBase* SourceTranslator = UInterchangeManager::GetInterchangeManager().GetTranslatorForSourceData(PayloadSourceData);
-	FGCObjectScopeGuard ScopedSourceTranslator(SourceTranslator);
-	const IInterchangeTextureLightProfilePayloadInterface* TextureTranslator = Cast< IInterchangeTextureLightProfilePayloadInterface >(SourceTranslator);
+	UE::Interchange::Private::FScopedTranslator ScopedTranslator(PayloadKey, Results);
+	const IInterchangeTextureLightProfilePayloadInterface* TextureTranslator = ScopedTranslator.GetPayLoadInterface< IInterchangeTextureLightProfilePayloadInterface >();
 	if (!ensure(TextureTranslator))
 	{
 		return TOptional<UE::Interchange::FImportLightProfile>();
 	}
 
-	SourceTranslator->ClearInternalFlags(EInternalObjectFlags::Async);
-	SourceTranslator->SetResultsContainer(Results);
-
 	AlternateTexturePath = PayloadKey;
 
 	AlternateTexturePath = PayloadKey;
diff --git a/Engine/Plugins/Interchange/Extensions/OpenUSD/Source/Import/Private/InterchangeUsdTranslator.cpp b/Engine/Plugins/Interchange/Extensions/OpenUSD/Source/Import/Private/InterchangeUsdTranslator.cpp
--- a/Engine/Plugins/Interchange/Extensions/OpenUSD/Source/Import/Private/InterchangeUsdTranslator.cpp
+++ b/Engine/Plugins/Interchange/Extensions/OpenUSD/Source/Import/Private/InterchangeUsdTranslator.cpp
@@ -4655,7 +4655,8 @@ namespace UE::InterchangeUsdTranslator::Private
 		const FString& PayloadKey,
 		TOptional<FString>& AlternateTexturePath,
 		const UInterchangeTranslatorBase* Translator,
-		UInterchangeUSDTranslatorImpl& Impl
+		UInterchangeUSDTranslatorImpl& Impl,
+		UInterchangeResultsContainer* Results
 	)
 	{
 		TRACE_CPUPROFILER_EVENT_SCOPE(GetMaterialXTexturePayloadData)
@@ -4673,25 +4674,14 @@ namespace UE::InterchangeUsdTranslator::Private
 			CompressionSettings = TextureCompressionSettings(FCString::Atoi(*PayloadKey.Mid(IndexTextureCompression + 1)));
 		}
 #endif
-		UInterchangeSourceData* PayloadSourceData = UInterchangeManager::GetInterchangeManager().CreateSourceData(Filename);
-		FGCObjectScopeGuard ScopedSourceData(PayloadSourceData);
-
-		if (!PayloadSourceData)
-		{
-			return TexturePayloadData;
-		}
-
-		UInterchangeTranslatorBase* SourceTranslator = UInterchangeManager::GetInterchangeManager().GetTranslatorForSourceData(PayloadSourceData);
-		FGCObjectScopeGuard ScopedSourceTranslator(SourceTranslator);
-		const IInterchangeTexturePayloadInterface* TextureTranslator = Cast<IInterchangeTexturePayloadInterface>(SourceTranslator);
+		UE::Interchange::Private::FScopedTranslator ScopedTranslator(Filename, Results);
+		const IInterchangeTexturePayloadInterface* TextureTranslator = ScopedTranslator.GetPayLoadInterface<IInterchangeTexturePayloadInterface>();
 
 		if (!TextureTranslator)
 		{
 			return TexturePayloadData;
 		}
 
-		SourceTranslator->SetResultsContainer(Translator->Results);
-
 		AlternateTexturePath = Impl.GetTextureSourcePath(Filename);
 
 		TexturePayloadData = TextureTranslator->GetTexturePayloadData(PayloadKey, AlternateTexturePath);
@@ -4770,7 +4760,7 @@ TOptional<UE::Interchange::FImportImage> UInterchangeUSDTranslator::GetTexturePa
 	// copying the behavior from InterchangeMaterialXTranslator
 	if (!TexturePayloadData.IsSet())
 	{
-		TexturePayloadData = GetMaterialXTexturePayloadData(PayloadKey, AlternateTexturePath, this, *ImplPtr);
+		TexturePayloadData = GetMaterialXTexturePayloadData(PayloadKey, AlternateTexturePath, this, *ImplPtr, Results);
 	}
 
 #endif	  // USE_USD_SDK
diff --git a/Engine/Plugins/Interchange/Runtime/Source/Import/Private/MaterialX/InterchangeMaterialXTranslator.cpp b/Engine/Plugins/Interchange/Runtime/Source/Import/Private/MaterialX/InterchangeMaterialXTranslator.cpp
--- a/Engine/Plugins/Interchange/Runtime/Source/Import/Private/MaterialX/InterchangeMaterialXTranslator.cpp
+++ b/Engine/Plugins/Interchange/Runtime/Source/Import/Private/MaterialX/InterchangeMaterialXTranslator.cpp
@@ -2,9 +2,9 @@
 
 #include "MaterialX/InterchangeMaterialXTranslator.h"
 #include "InterchangeManager.h"
+#include "InterchangeTranslatorHelper.h"
 #include "MaterialX/MaterialXUtils/MaterialXManager.h"
 #include "Nodes/InterchangeSourceNode.h"
-#include "UObject/GCObjectScopeGuard.h"
 
 #if WITH_EDITOR
 #include "MaterialXFormat/Util.h"
@@ -202,25 +202,15 @@ TOptional<UE::Interchange::FImportImage> UInterchangeMaterialXTranslator::GetTex
 		CompressionSettings = TextureCompressionSettings(FCString::Atoi(*PayloadKey.Mid(IndexTextureCompression + 1)));
 	}
 #endif
-	UInterchangeSourceData* PayloadSourceData = UInterchangeManager::GetInterchangeManager().CreateSourceData(Filename);
-	FGCObjectScopeGuard ScopedSourceData(PayloadSourceData);
 
-	if(!PayloadSourceData)
-	{
-		return TOptional<UE::Interchange::FImportImage>();
-	}
-
-	UInterchangeTranslatorBase* SourceTranslator = UInterchangeManager::GetInterchangeManager().GetTranslatorForSourceData(PayloadSourceData);
-	FGCObjectScopeGuard ScopedSourceTranslator(SourceTranslator);
-	const IInterchangeTexturePayloadInterface* TextureTranslator = Cast<IInterchangeTexturePayloadInterface>(SourceTranslator);
+	UE::Interchange::Private::FScopedTranslator ScopedTranslator(Filename, Results);
+	const IInterchangeTexturePayloadInterface* TextureTranslator = ScopedTranslator.GetPayLoadInterface<IInterchangeTexturePayloadInterface>();
 
 	if(!ensure(TextureTranslator))
 	{
 		return TOptional<UE::Interchange::FImportImage>();
 	}
 
-	SourceTranslator->SetResultsContainer(Results);
-
 	AlternateTexturePath = Filename;
 
 	TOptional<UE::Interchange::FImportImage> TexturePayloadData = TextureTranslator->GetTexturePayloadData(PayloadKey, AlternateTexturePath);
