[Interchange] Bug fix for FindFactoryNodeByUniqueID in Datasmith Utils.

#rnx

diff --git a/Engine/Plugins/Experimental/Enterprise/DatasmithInterchange/Source/DatasmithInterchange/Private/InterchangeDatasmithUtils.h b/Engine/Plugins/Experimental/Enterprise/DatasmithInterchange/Source/DatasmithInterchange/Private/InterchangeDatasmithUtils.h
--- a/Engine/Plugins/Experimental/Enterprise/DatasmithInterchange/Source/DatasmithInterchange/Private/InterchangeDatasmithUtils.h
+++ b/Engine/Plugins/Experimental/Enterprise/DatasmithInterchange/Source/DatasmithInterchange/Private/InterchangeDatasmithUtils.h
@@ -107,7 +107,7 @@ namespace UE::DatasmithInterchange::NodeUtils
 			{
 				if (NodeUniqueID == NodeUid)
 				{
-					Cast<T>(Node);
+					FactoryNode = Cast<T>(Node);
 					return true;
 				}
 				return false;
