Fixed linkage error when CADKernelEngine module is included
- Solution: Temporarily disable CAD import through DatasmithInterchange at runtime

#rnx
#Synced-CL 48455087

diff --git a/Engine/Plugins/Experimental/Enterprise/DatasmithInterchange/DatasmithInterchange.uplugin b/Engine/Plugins/Experimental/Enterprise/DatasmithInterchange/DatasmithInterchange.uplugin
--- a/Engine/Plugins/Experimental/Enterprise/DatasmithInterchange/DatasmithInterchange.uplugin
+++ b/Engine/Plugins/Experimental/Enterprise/DatasmithInterchange/DatasmithInterchange.uplugin
@@ -2,7 +2,7 @@
 	"FileVersion": 1,
 	"Version": 1,
 	"VersionName": "1.0",
-	"FriendlyName": "Datasmith Interchange",
+	"FriendlyName": "Interchange Datasmith",
 	"Description": "Interchange Importer for Datasmith.",
 	"Category": "Importers",
 	"CreatedBy": "Epic Games, Inc.",
@@ -38,7 +38,7 @@
 		},
 		{
 			"Name": "DatasmithCADImporter",
-			"Enabled": true,
+			"Enabled": false,
 			"SupportedTargetPlatforms": [
 				"Win64"
 			]
diff --git a/Engine/Plugins/Experimental/Enterprise/DatasmithInterchange/Source/DatasmithInterchange/DatasmithInterchange.Build.cs b/Engine/Plugins/Experimental/Enterprise/DatasmithInterchange/Source/DatasmithInterchange/DatasmithInterchange.Build.cs
--- a/Engine/Plugins/Experimental/Enterprise/DatasmithInterchange/Source/DatasmithInterchange/DatasmithInterchange.Build.cs
+++ b/Engine/Plugins/Experimental/Enterprise/DatasmithInterchange/Source/DatasmithInterchange/DatasmithInterchange.Build.cs
@@ -10,7 +10,6 @@ namespace UnrealBuildTool.Rules
 			PrivateDependencyModuleNames.AddRange(
 				new string[]
 				{
-					"CADTools",
 					"DatasmithContent",
 					"DatasmithCore",
 					"DatasmithTranslator",
@@ -23,7 +22,6 @@ namespace UnrealBuildTool.Rules
 					"InterchangeImport",
 					"InterchangePipelines",
 					"MeshDescription",
-					"ParametricSurface",
 					"StaticMeshDescription"
 				}
 			);
@@ -45,16 +43,23 @@ namespace UnrealBuildTool.Rules
 					}
 				);
 			}
+			
+			if(Target.LinkType != TargetLinkType.Monolithic)
+			{
+				PrivateDefinitions.Add("NO_MONOLITHIC_BUILD");
+				PrivateDependencyModuleNames.Add("CADTools");
+				PrivateDependencyModuleNames.Add("ParametricSurface");
+			}
 
 			PublicDependencyModuleNames.AddRange(
-				new string[]
-				{
+					new string[]
+					{
 					"Core",
 					"CoreUObject",
 					"InterchangeCore",
 					"InterchangeImport",
-				}
-			);
+					}
+				);
 		}
 	}
 }
\ No newline at end of file
diff --git a/Engine/Plugins/Experimental/Enterprise/DatasmithInterchange/Source/DatasmithInterchange/Private/InterchangeDatasmithStaticMeshPipeline.cpp b/Engine/Plugins/Experimental/Enterprise/DatasmithInterchange/Source/DatasmithInterchange/Private/InterchangeDatasmithStaticMeshPipeline.cpp
--- a/Engine/Plugins/Experimental/Enterprise/DatasmithInterchange/Source/DatasmithInterchange/Private/InterchangeDatasmithStaticMeshPipeline.cpp
+++ b/Engine/Plugins/Experimental/Enterprise/DatasmithInterchange/Source/DatasmithInterchange/Private/InterchangeDatasmithStaticMeshPipeline.cpp
@@ -14,7 +14,9 @@
 
 #include "Engine/StaticMesh.h"
 #include "Materials/MaterialInterface.h"
+#ifdef NO_MONOLITHIC_BUILD
 #include "DatasmithParametricSurfaceData.h"
+#endif
 
 void UInterchangeDatasmithStaticMeshPipeline::ExecutePipeline(UInterchangeBaseNodeContainer* NodeContainer, const TArray<UInterchangeSourceData*>& InSourceDatas, const FString& ContentBasePath)
 {
@@ -133,6 +135,7 @@ void UInterchangeDatasmithStaticMeshPipeline::ApplyAdditionalData(const UInterch
 		return;
 	}
 
+#ifdef NO_MONOLITHIC_BUILD
 	if (const TObjectPtr<UDatasmithAdditionalData>* AdditionalDataPtr = StaticMeshDataNode->AdditionalDataMap.Find(MeshNodeUid))
 	{
 		if (UDatasmithParametricSurfaceData* AdditionalData = Cast<UDatasmithParametricSurfaceData>(AdditionalDataPtr->Get()))
@@ -169,4 +172,5 @@ void UInterchangeDatasmithStaticMeshPipeline::ApplyAdditionalData(const UInterch
 		}
 	}
 #endif
+#endif
 }
diff --git a/Engine/Plugins/Experimental/Enterprise/DatasmithInterchange/Source/DatasmithInterchange/Private/InterchangeDatasmithTranslator.cpp b/Engine/Plugins/Experimental/Enterprise/DatasmithInterchange/Source/DatasmithInterchange/Private/InterchangeDatasmithTranslator.cpp
--- a/Engine/Plugins/Experimental/Enterprise/DatasmithInterchange/Source/DatasmithInterchange/Private/InterchangeDatasmithTranslator.cpp
+++ b/Engine/Plugins/Experimental/Enterprise/DatasmithInterchange/Source/DatasmithInterchange/Private/InterchangeDatasmithTranslator.cpp
@@ -18,9 +18,11 @@
 #include "DatasmithUtils.h"
 #include "DatasmithVariantElements.h"
 #include "IDatasmithSceneElements.h"
+#ifdef NO_MONOLITHIC_BUILD
 #include "DatasmithParametricSurfaceData.h"
 
 #include "CADOptions.h"
+#endif
 #include "ExternalSourceModule.h"
 #include "SourceUri.h"
 #include "InterchangeCameraNode.h"
@@ -202,12 +204,14 @@ bool UInterchangeDatasmithTranslator::Translate(UInterchangeBaseNodeContainer& B
 
 	TSharedPtr<IDatasmithScene> DatasmithScene;
 	{
+#ifdef NO_MONOLITHIC_BUILD
 		TGuardValue<bool> EnableCADCache(CADLibrary::FImportParameters::bGEnableCADCache, true);
-
+#endif
 		if (GetSettings())
 		{
+#ifdef NO_MONOLITHIC_BUILD
 			CADLibrary::FImportParameters::bGEnableCADCache = true;
-
+#endif
 			const TSharedPtr<IDatasmithTranslator>& DatasmithTranslator = LoadedExternalSource->GetAssetTranslator();
 			if (DatasmithTranslator)
 			{
@@ -1063,6 +1067,7 @@ bool UInterchangeDatasmithTranslator::GetMeshDescription(const TSharedPtr<IDatas
 
 	if (DatasmithMeshPayload.LodMeshes.Num() > 0)
 	{
+#ifdef NO_MONOLITHIC_BUILD
 		for (UDatasmithAdditionalData* AdditionalData : DatasmithMeshPayload.AdditionalData)
 		{
 			if (UDatasmithParametricSurfaceData* ParametricSurfaceData = Cast<UDatasmithParametricSurfaceData>(AdditionalData))
@@ -1074,6 +1079,7 @@ bool UInterchangeDatasmithTranslator::GetMeshDescription(const TSharedPtr<IDatas
 				break;
 			}
 		}
+#endif
 
 		if (!FStaticMeshOperations::ValidateAndFixData(DatasmithMeshPayload.LodMeshes[0], MeshElement->GetName()))
 		{
