Converted QA_ContentPipeline project into a C++ project

#rnx
#Synced-CL 48455087

diff --git a/Engine/Plugins/Experimental/Enterprise/DatasmithInterchange/DatasmithInterchange.uplugin b/Engine/Plugins/Experimental/Enterprise/DatasmithInterchange/DatasmithInterchange.uplugin
--- a/Engine/Plugins/Experimental/Enterprise/DatasmithInterchange/DatasmithInterchange.uplugin
+++ b/Engine/Plugins/Experimental/Enterprise/DatasmithInterchange/DatasmithInterchange.uplugin
@@ -38,7 +38,7 @@
 		},
 		{
 			"Name": "DatasmithCADImporter",
-			"Enabled": false,
+			"Enabled": true,
 			"SupportedTargetPlatforms": [
 				"Win64"
 			]
